Changed all instances of notifications with a "Show Log" hyperlink to be more precise about which log they will open.
#ttp 338544 - EDITOR: Update Notifications "Show Log" to "Show Output Log"
#branch UE4
#change Changed hyperlink text to be more specific in individual instances of FNotificationInfo objects.

diff --git a/Engine/Source/Developer/MessageLog/Private/Presentation/MessageLogListingViewModel.cpp b/Engine/Source/Developer/MessageLog/Private/Presentation/MessageLogListingViewModel.cpp
--- a/Engine/Source/Developer/MessageLog/Private/Presentation/MessageLogListingViewModel.cpp
+++ b/Engine/Source/Developer/MessageLog/Private/Presentation/MessageLogListingViewModel.cpp
@@ -312,6 +312,7 @@ void FMessageLogListingViewModel::NotifyIfAnyMessages( const FText& Message, EMe
 		ErrorNotification.Image = FEditorStyle::GetBrush(FTokenizedMessage::GetSeverityIconName(HighestSeverityPresent(0)));
 		ErrorNotification.bFireAndForget = true;
 		ErrorNotification.Hyperlink = FSimpleDelegate::CreateSP(this, &FMessageLogListingViewModel::OpenMessageLog);
+		ErrorNotification.HyperlinkText = LOCTEXT("ShowMessageLogHyperlink", "Show Message Log");
 		ErrorNotification.ExpireDuration = 8.0f; // Need this message to last a little longer than normal since the user may want to "Show Log"
 		ErrorNotification.bUseThrobber = true;
 
diff --git a/Engine/Source/Editor/Blutility/Private/PlacedEditorUtilityBase.cpp b/Engine/Source/Editor/Blutility/Private/PlacedEditorUtilityBase.cpp
--- a/Engine/Source/Editor/Blutility/Private/PlacedEditorUtilityBase.cpp
+++ b/Engine/Source/Editor/Blutility/Private/PlacedEditorUtilityBase.cpp
@@ -32,6 +32,7 @@ FNotificationInfo ErrorNotification(TEXT(""));
 ErrorNotification.Image = FEditorStyle::GetBrush(TEXT("MessageLog.Error"));
 ErrorNotification.bFireAndForget = true;
 ErrorNotification.Hyperlink = FSimpleDelegate::CreateRaw(&Local::OpenMessageLog);
+ErrorNotification.HyperlinkText = LOCTEXT("ShowMessageLogHyperlink", "Show Message Log");
 ErrorNotification.ExpireDuration = 3.0f; // Need this message to last a little longer than normal since the user may want to "Show Log"
 ErrorNotification.bUseThrobber = true;
 ErrorNotification.Text = stuff;
diff --git a/Engine/Source/Editor/ContentBrowser/Private/ContentBrowserUtils.cpp b/Engine/Source/Editor/ContentBrowser/Private/ContentBrowserUtils.cpp
--- a/Engine/Source/Editor/ContentBrowser/Private/ContentBrowserUtils.cpp
+++ b/Engine/Source/Editor/ContentBrowser/Private/ContentBrowserUtils.cpp
@@ -364,7 +364,7 @@ bool ContentBrowserUtils::LoadAssetsIfNeeded(const TArray<FString>& ObjectPaths,
 			FNotificationInfo Info(LOCTEXT("LoadObjectFailed", "Failed to load assets"));
 			Info.ExpireDuration = 5.0f;
 			Info.Hyperlink = FSimpleDelegate::CreateStatic([](){ FMessageLog("LoadErrors").Open(EMessageSeverity::Info, true); });
-			Info.HyperlinkText = LOCTEXT("LoadObjectHyperlink", "Show Log");
+			Info.HyperlinkText = LOCTEXT("LoadObjectHyperlink", "Show Message Log");
 
 			FSlateNotificationManager::Get().AddNotification(Info);
 			return false;
diff --git a/Engine/Source/Editor/MainFrame/Private/Frame/MainFrameActions.cpp b/Engine/Source/Editor/MainFrame/Private/Frame/MainFrameActions.cpp
--- a/Engine/Source/Editor/MainFrame/Private/Frame/MainFrameActions.cpp
+++ b/Engine/Source/Editor/MainFrame/Private/Frame/MainFrameActions.cpp
@@ -972,6 +972,7 @@ void FMainFrameActionCallbacks::CreateUatTask( const FString& CommandLine, const
 	Info.bFireAndForget = false;
 	Info.ExpireDuration = 3.0f;
 	Info.Hyperlink = FSimpleDelegate::CreateStatic(&FMainFrameActionCallbacks::HandleUatHyperlinkNavigate);
+	Info.HyperlinkText = LOCTEXT("ShowOutputLogHyperlink", "Show Output Log");
 	Info.ButtonDetails.Add(
 		FNotificationButtonInfo(
 			LOCTEXT("UatTaskCancel", "Cancel"),
diff --git a/Engine/Source/Editor/SourceControlWindows/Private/SSourceControlSubmit.cpp b/Engine/Source/Editor/SourceControlWindows/Private/SSourceControlSubmit.cpp
--- a/Engine/Source/Editor/SourceControlWindows/Private/SSourceControlSubmit.cpp
+++ b/Engine/Source/Editor/SourceControlWindows/Private/SSourceControlSubmit.cpp
@@ -559,7 +559,7 @@ bool FSourceControlWindows::PromptForCheckin(const TArray<FString>& InPackageNam
 					// report success with a notification
 					FNotificationInfo Info(CheckInOperation->GetSuccessMessage());
 					Info.ExpireDuration = 8.0f;
-					Info.HyperlinkText = LOCTEXT("SCC_Checkin_ShowLog", "Show Log");
+					Info.HyperlinkText = LOCTEXT("SCC_Checkin_ShowLog", "Show Message Log");
 					Info.Hyperlink = FSimpleDelegate::CreateStatic([](){ FMessageLog("SourceControl").Open(EMessageSeverity::Info, true); });
 					FSlateNotificationManager::Get().AddNotification(Info);
 
diff --git a/Engine/Source/Editor/TranslationEditor/Private/TranslationDataManager.cpp b/Engine/Source/Editor/TranslationEditor/Private/TranslationDataManager.cpp
--- a/Engine/Source/Editor/TranslationEditor/Private/TranslationDataManager.cpp
+++ b/Engine/Source/Editor/TranslationEditor/Private/TranslationDataManager.cpp
@@ -328,6 +328,7 @@ bool FTranslationDataManager::WriteJSONToTextFile(TSharedRef<FJsonObject>& Outpu
 		FNotificationInfo Info( LOCTEXT("ArchiveCheckedOut", "Translation Archive Successfully Checked Out and Saved.") );
 		Info.ExpireDuration = 5;
 		Info.Hyperlink = FSimpleDelegate::CreateStatic(&Local::OpenSourceControlMessageLog);
+		Info.HyperlinkText = LOCTEXT("ShowMessageLogHyperlink", "Show Message Log");
 		Info.bFireAndForget = true;
 		Info.bUseSuccessFailIcons = true;
 		Info.Image = FEditorStyle::GetBrush(TEXT("NotificationList.SuccessImage"));
diff --git a/Engine/Source/Editor/UnrealEd/Private/PlayLevel.cpp b/Engine/Source/Editor/UnrealEd/Private/PlayLevel.cpp
--- a/Engine/Source/Editor/UnrealEd/Private/PlayLevel.cpp
+++ b/Engine/Source/Editor/UnrealEd/Private/PlayLevel.cpp
@@ -1281,6 +1281,7 @@ void UEditorEngine::PlayUsingLauncher()
 		Info.bFireAndForget = false;
 		Info.ExpireDuration = 3.0f;
 		Info.Hyperlink = FSimpleDelegate::CreateStatic(HandleHyperlinkNavigate);
+		Info.HyperlinkText = LOCTEXT("ShowOutputLogHyperlink", "Show Output Log");
 		Info.ButtonDetails.Add(
 			FNotificationButtonInfo(
 				LOCTEXT("LauncherTaskCancel", "Cancel"),
