Fix for message string in FMessageLogListingModel::AddMessageInternal() being treated as a formatting string

diff --git a/Engine/Source/Developer/MessageLog/Private/Model/MessageLogListingModel.cpp b/Engine/Source/Developer/MessageLog/Private/Model/MessageLogListingModel.cpp
--- a/Engine/Source/Developer/MessageLog/Private/Model/MessageLogListingModel.cpp
+++ b/Engine/Source/Developer/MessageLog/Private/Model/MessageLogListingModel.cpp
@@ -67,7 +67,7 @@ void FMessageLogListingModel::AddMessageInternal( const TSharedRef<FTokenizedMes
 	{
 		SET_WARN_COLOR(LogColor);
 	}
-	FMsg::Logf(__FILE__, __LINE__, *LogName.ToString(), FMessageLog::GetLogVerbosity(NewMessage->GetSeverity()), *NewMessage->ToText().ToString());
+	FMsg::Logf(__FILE__, __LINE__, *LogName.ToString(), FMessageLog::GetLogVerbosity(NewMessage->GetSeverity()), TEXT("%s"), *NewMessage->ToText().ToString());
 	CLEAR_WARN_COLOR();
 }
 
