Removed root window check in FMessageLogModule::OpenMessageLog() on Mac, as it's always NULL

diff --git a/Engine/Source/Developer/MessageLog/Private/MessageLogModule.cpp b/Engine/Source/Developer/MessageLog/Private/MessageLogModule.cpp
--- a/Engine/Source/Developer/MessageLog/Private/MessageLogModule.cpp
+++ b/Engine/Source/Developer/MessageLog/Private/MessageLogModule.cpp
@@ -97,7 +97,11 @@ TSharedRef<SWidget> FMessageLogModule::CreateLogListingWidget(const TSharedRef<I
 void FMessageLogModule::OpenMessageLog(const FName& LogName)
 {
 	// only open the message log if we have a window created for the tab manager & our delegate allows it
-	if(bCanDisplayMessageLog && FGlobalTabmanager::Get()->GetRootWindow().IsValid())
+	if(bCanDisplayMessageLog
+#if !PLATFORM_MAC
+	   && FGlobalTabmanager::Get()->GetRootWindow().IsValid()
+#endif
+	   )
 	{
 		FGlobalTabmanager::Get()->InvokeTab(FName("MessageLog"));
 		MessageLogViewModel->ChangeCurrentListingViewModel(LogName);
