Fix to self-assignment typo caught by Mac build

diff --git a/Engine/Source/Developer/MessageLog/Private/Presentation/MessageLogListingViewModel.cpp b/Engine/Source/Developer/MessageLog/Private/Presentation/MessageLogListingViewModel.cpp
--- a/Engine/Source/Developer/MessageLog/Private/Presentation/MessageLogListingViewModel.cpp
+++ b/Engine/Source/Developer/MessageLog/Private/Presentation/MessageLogListingViewModel.cpp
@@ -376,7 +376,7 @@ bool FMessageLogListingViewModel::GetShowPages() const
 
 void FMessageLogListingViewModel::SetAllowClear(bool bInAllowClear)
 {
-	bAllowClear = bAllowClear;
+	bAllowClear = bInAllowClear;
 }
 
 bool FMessageLogListingViewModel::GetAllowClear() const
