MessageLog: Message log burst upgrades

- Redesigned the layout of the message log viewer
- Added 'Action' message token type for actionable items
- Started to remove message token delegates in favor of Action tokens
- NULL to nullptr

diff --git a/Engine/Source/Developer/MessageLog/Private/MessageLogModule.cpp b/Engine/Source/Developer/MessageLog/Private/MessageLogModule.cpp
--- a/Engine/Source/Developer/MessageLog/Private/MessageLogModule.cpp
+++ b/Engine/Source/Developer/MessageLog/Private/MessageLogModule.cpp
@@ -14,7 +14,7 @@
 
 IMPLEMENT_MODULE( FMessageLogModule, MessageLog );
 
-TSharedRef<SDockTab> SpawnMessageLog( const FSpawnTabArgs& Args, TSharedPtr<FMessageLogViewModel> MessageLogViewModel )
+TSharedRef<SDockTab> SpawnMessageLog( const FSpawnTabArgs& Args, TSharedRef<FMessageLogViewModel> MessageLogViewModel )
 {
 	return SNew(SDockTab)
 		.TabRole(ETabRole::NomadTab)
@@ -40,7 +40,7 @@ void FMessageLogModule::StartupModule()
 	MessageLogViewModel->Initialize();
 
 #if WITH_EDITOR
-	FGlobalTabmanager::Get()->RegisterNomadTabSpawner("MessageLog", FOnSpawnTab::CreateStatic( &SpawnMessageLog, MessageLogViewModel ))
+	FGlobalTabmanager::Get()->RegisterNomadTabSpawner("MessageLog", FOnSpawnTab::CreateStatic( &SpawnMessageLog, MessageLogViewModel.ToSharedRef() ))
 		.SetDisplayName(NSLOCTEXT("UnrealEditor", "MessageLogTab", "Message Log"))
 		.SetTooltipText(NSLOCTEXT("UnrealEditor", "MessageLogTooltipText", "Open the Message Log tab."))
 		.SetGroup( WorkspaceMenu::GetMenuStructure().GetToolsCategory() )
diff --git a/Engine/Source/Developer/MessageLog/Private/Presentation/MessageLogListingViewModel.h b/Engine/Source/Developer/MessageLog/Private/Presentation/MessageLogListingViewModel.h
--- a/Engine/Source/Developer/MessageLog/Private/Presentation/MessageLogListingViewModel.h
+++ b/Engine/Source/Developer/MessageLog/Private/Presentation/MessageLogListingViewModel.h
@@ -7,11 +7,16 @@
 #include "IMessageLogListing.h"
 #include "MessageLogInitializationOptions.h"
 
-/** The non-UI solution specific presentation logic for a collection of messages for a particular system */
-class FMessageLogListingViewModel : public IMessageLogListing, public TSharedFromThis< FMessageLogListingViewModel >
-{
 
+/**
+ * The non-UI solution specific presentation logic for a collection of messages for a particular system.
+ */
+class FMessageLogListingViewModel
+	: public TSharedFromThis<FMessageLogListingViewModel>
+	, public IMessageLogListing
+{
 public:
+
 	/**  
 	 *	Factory method which creates a new FMessageLogListingViewModel object
 	 *
diff --git a/Engine/Source/Developer/MessageLog/Private/Presentation/MessageLogViewModel.h b/Engine/Source/Developer/MessageLog/Private/Presentation/MessageLogViewModel.h
--- a/Engine/Source/Developer/MessageLog/Private/Presentation/MessageLogViewModel.h
+++ b/Engine/Source/Developer/MessageLog/Private/Presentation/MessageLogViewModel.h
@@ -76,7 +76,7 @@ public:
 	FString GetCurrentListingLabel() const;
 
 	/** Get the linearized array of ViewModels */
-	const TArray< TSharedPtr< class FMessageLogListingViewModel > >& GetLogListingViewModels() const { return ViewModelArray; }
+	const TArray<IMessageLogListingPtr>& GetLogListingViewModels() const { return ViewModelArray; }
 
 private:
 
@@ -92,7 +92,7 @@ private:
 	TMap< FName, TSharedPtr< class FMessageLogListingViewModel > > NameToViewModelMap;
 
 	/** A linearized array of the ViewModels - we need this for the data to display in a SComboBox */
-	TArray< TSharedPtr< class FMessageLogListingViewModel > > ViewModelArray;
+	TArray<IMessageLogListingPtr> ViewModelArray;
 
 	/** The currently selected log listing */
 	TSharedPtr< class FMessageLogListingViewModel > SelectedLogListingViewModel;
diff --git a/Engine/Source/Developer/MessageLog/Private/UserInterface/SMessageLog.cpp b/Engine/Source/Developer/MessageLog/Private/UserInterface/SMessageLog.cpp
--- a/Engine/Source/Developer/MessageLog/Private/UserInterface/SMessageLog.cpp
+++ b/Engine/Source/Developer/MessageLog/Private/UserInterface/SMessageLog.cpp
@@ -5,124 +5,124 @@
 #include "MessageLogViewModel.h"
 #include "MessageLogListingViewModel.h"
 #include "SMessageLogListing.h"
+#include "SMessageLogCategoryListRow.h"
 
-#define LOCTEXT_NAMESPACE "Developer.MessageLog"
 
-SMessageLog::SMessageLog()
-	: bUpdatingSelection( false )
-{
+#define LOCTEXT_NAMESPACE "SMessageLog"
 
-}
+
+/* SMessageLog structors
+ *****************************************************************************/
 
 SMessageLog::~SMessageLog()
 {
-	MessageLogViewModel->OnSelectionChanged().RemoveAll( this );
+	//ViewModel->OnSelectionChanged().RemoveAll(this);
 }
 
+
 BEGIN_SLATE_FUNCTION_BUILD_OPTIMIZATION
-void SMessageLog::Construct( const FArguments& InArgs, const TSharedPtr<FMessageLogViewModel>& InMessageLogViewModel )
+void SMessageLog::Construct( const FArguments& InArgs, const TSharedRef<FMessageLogViewModel>& InViewModel )
 {
-	MessageLogViewModel = InMessageLogViewModel;
-	check(MessageLogViewModel.IsValid());
+	ViewModel = InViewModel;
 
 	ChildSlot
 	[
-		SNew(SVerticalBox)
-		+SVerticalBox::Slot().AutoHeight() .Padding(2)
-		[
-			SNew(SHorizontalBox)
-			+SHorizontalBox::Slot()
-			.AutoWidth()
+		SNew(SSplitter)
+			.Orientation(Orient_Horizontal)
+
+		+ SSplitter::Slot()
+			.Value(0.2f)
 			[
-				SAssignNew(CurrentListingSelector, SComboBox< TSharedPtr<FMessageLogListingViewModel> >)
-					.ContentPadding(3)
-					.OptionsSource(&MessageLogViewModel->GetLogListingViewModels())
-					.OnSelectionChanged(this, &SMessageLog::ChangeCurrentListingSelection)
-					.OnGenerateWidget(this, &SMessageLog::MakeCurrentListingSelectionWidget)
+				SNew(SBorder)
+					.BorderImage(FEditorStyle::GetBrush("ToolPanel.GroupBorder"))
+					.Padding(0.0f)
 					[
-						SNew(STextBlock)
-						.Text(this, &SMessageLog::GetCurrentListingLabel)
+						// log categories list
+						SAssignNew(CategoriesListView, SListView<IMessageLogListingPtr>)
+							.ItemHeight(24.0f)
+							.ListItemsSource(&ViewModel->GetLogListingViewModels())
+							.OnGenerateRow(this, &SMessageLog::HandleCategoriesListGenerateRow)
+							.OnSelectionChanged(this, &SMessageLog::HandleCategoriesListSelectionChanged)
+							.SelectionMode(ESelectionMode::Single)
+							.HeaderRow
+							(
+								SNew(SHeaderRow)
+									.Visibility(EVisibility::Collapsed)
+
+								+ SHeaderRow::Column("Name")
+									.DefaultLabel(LOCTEXT("CategoriesListNameColumnHeader", "Category"))
+									.FillWidth(1.0)
+							)
 					]
 			]
-		]
-		+SVerticalBox::Slot()
-		.FillHeight(1)
-		[
-			SAssignNew(CurrentListingDisplay, SBorder)
-		]
+
+		+ SSplitter::Slot()
+			.Value(0.8f)
+			[
+				SAssignNew(CurrentListingDisplay, SBorder)
+					.BorderImage(FEditorStyle::GetBrush("NoBorder"))
+					.Padding(0.0f)
+			]
 	];
 
 	// Default to the first available (or the last used) listing
-	const TArray< TSharedPtr< FMessageLogListingViewModel > >& ViewModels = MessageLogViewModel->GetLogListingViewModels();
-	if ( ViewModels.Num() > 0 )
+	const TArray<IMessageLogListingPtr>& ViewModels = ViewModel->GetLogListingViewModels();
+
+	if (ViewModels.Num() > 0)
 	{
-		TSharedPtr< FMessageLogListingViewModel > DefaultViewModel = ViewModels[0];
-		if ( FPaths::FileExists(GEditorUserSettingsIni) )
+		IMessageLogListingPtr DefaultViewModel = ViewModels[0];
+
+		if (FPaths::FileExists(GEditorUserSettingsIni))
 		{
 			FString LogName;
-			if ( GConfig->GetString( TEXT("MessageLog"), TEXT("LastLogListing"), LogName, GEditorUserSettingsIni ) )
+
+			if (GConfig->GetString(TEXT("MessageLog"), TEXT("LastLogListing"), LogName, GEditorUserSettingsIni))
 			{
-				TSharedPtr< FMessageLogListingViewModel > LoadedViewModel = MessageLogViewModel->FindLogListingViewModel( FName( *LogName ) );
-				if ( LoadedViewModel.IsValid() )
+				IMessageLogListingPtr LoadedViewModel = ViewModel->FindLogListingViewModel(FName(*LogName));
+				
+				if (LoadedViewModel.IsValid())
 				{
 					DefaultViewModel = LoadedViewModel;
 				}
 			}
 		}
-		ChangeCurrentListingSelection( DefaultViewModel, ESelectInfo::Direct );
-	}
 
-	MessageLogViewModel->OnSelectionChanged().AddSP( this, &SMessageLog::OnSelectionUpdated );
+		CategoriesListView->SetSelection(DefaultViewModel, ESelectInfo::Direct);
+	}
 
+//	ViewModel->OnSelectionChanged().AddSP(this, &SMessageLog::OnSelectionUpdated);
 }
 END_SLATE_FUNCTION_BUILD_OPTIMIZATION
 
-void SMessageLog::OnSelectionUpdated()
-{
-	TSharedPtr<FMessageLogListingViewModel> LogListingViewModel = MessageLogViewModel->GetCurrentListingViewModel();
 
-	if ( LogListingViewModel.IsValid() )
-	{
-		if( !bUpdatingSelection )
-		{
-			bUpdatingSelection = true;
-			CurrentListingSelector->SetSelectedItem( LogListingViewModel ) ;
-			bUpdatingSelection = false;
-		}
-
-		CurrentListingDisplay->SetContent(SAssignNew(LogListing, SMessageLogListing, LogListingViewModel.ToSharedRef()));
-	}
-}
+/* SWidget overrides
+ *****************************************************************************/
 
 FReply SMessageLog::OnKeyDown( const FGeometry& MyGeometry, const FKeyboardEvent& InKeyboardEvent )
 {
-	if(LogListing.IsValid())
+	if (LogListing.IsValid())
 	{
-		return LogListing->GetCommandList()->ProcessCommandBindings( InKeyboardEvent ) ? FReply::Handled() : FReply::Unhandled();
+		return LogListing->GetCommandList()->ProcessCommandBindings(InKeyboardEvent)
+			? FReply::Handled()
+			: FReply::Unhandled();
 	}
 	
 	return FReply::Unhandled();
 }
 
-FString SMessageLog::GetCurrentListingLabel() const
-{
-	return MessageLogViewModel->GetCurrentListingLabel();
-}
 
-void SMessageLog::ChangeCurrentListingSelection( TSharedPtr<FMessageLogListingViewModel> Selection, ESelectInfo::Type /*SelectInfo*/ )
+/* SMessageLog callbacks
+ *****************************************************************************/
+
+TSharedRef<ITableRow> SMessageLog::HandleCategoriesListGenerateRow( IMessageLogListingPtr Item, const TSharedRef<STableViewBase>& OwnerTable )
 {
-	if( bUpdatingSelection )
-	{
-		return;
-	}
-	bUpdatingSelection = true;
-	MessageLogViewModel->ChangeCurrentListingViewModel( Selection->GetName() );
-	bUpdatingSelection = false;
+	return SNew(SMessageLogCategoryListRow, Item.ToSharedRef(), OwnerTable);
 }
 
-TSharedRef<SWidget> SMessageLog::MakeCurrentListingSelectionWidget( TSharedPtr<FMessageLogListingViewModel> Selection )
+	
+void SMessageLog::HandleCategoriesListSelectionChanged( IMessageLogListingPtr Selection, ESelectInfo::Type SelectInfo )
 {
-	return SNew(STextBlock) .Text( Selection->GetLabel() );
+	CurrentListingDisplay->SetContent(SAssignNew(LogListing, SMessageLogListing, Selection.ToSharedRef()));
 }
 
 
diff --git a/Engine/Source/Developer/MessageLog/Private/UserInterface/SMessageLog.h b/Engine/Source/Developer/MessageLog/Private/UserInterface/SMessageLog.h
--- a/Engine/Source/Developer/MessageLog/Private/UserInterface/SMessageLog.h
+++ b/Engine/Source/Developer/MessageLog/Private/UserInterface/SMessageLog.h
@@ -4,64 +4,55 @@
 
 /**
  * An implementation for the message log widget.
+ *
  * It holds a series of message log listings which it can switch between.
  */
-class SMessageLog : public SCompoundWidget
+class SMessageLog
+	: public SCompoundWidget
 {
 public:
+
 	SLATE_BEGIN_ARGS(SMessageLog){}
 	SLATE_END_ARGS()
 
 	static FName AppName;
 	
-	/** Constructor */
-	SMessageLog();
-
-	/** Destructor */	
+	/** Destructor. */
 	~SMessageLog();
 
-	void Construct( const FArguments& InArgs, const TSharedPtr<class FMessageLogViewModel>& InMessageLogViewModel );
-
 	/**
-	 * Called after a key is pressed when this widget has keyboard focus
+	 * Constructs the widget.
 	 *
-	 * @param  InKeyboardEvent  Keyboard event
-	 *
-	 * @return  Returns whether the event was handled, along with other possible actions
+	 * @param InArgs The construction arguments.
+	 * @param InModel The view model to use.
 	 */
-	FReply OnKeyDown( const FGeometry& MyGeometry, const FKeyboardEvent& InKeyboardEvent ) override;
+	void Construct( const FArguments& InArgs, const TSharedRef<class FMessageLogViewModel>& InViewModel );
 
-private:
-	/** Gets the currently visible log listing's label */
-	FString GetCurrentListingLabel() const;
-
-	/** Changes the selected log listing */
-	void ChangeCurrentListingSelection( TSharedPtr<class FMessageLogListingViewModel> Selection, ESelectInfo::Type SelectInfo );
+public:
 
-	/** Generates the widgets for the log listings in the current listing combo box */
-	TSharedRef<SWidget> MakeCurrentListingSelectionWidget( TSharedPtr<class FMessageLogListingViewModel> Selection );
+	// SWidget overrides
 
-	/** Called whenever the viewmodel selection changes */
-	void OnSelectionUpdated();
+	FReply OnKeyDown( const FGeometry& MyGeometry, const FKeyboardEvent& InKeyboardEvent ) override;
 
-private:	
+private:
 
-	/** The message log view model */
-	TSharedPtr<class FMessageLogViewModel> MessageLogViewModel;
+	/** Callback for generating a row in the categories list. */
+	TSharedRef<ITableRow> HandleCategoriesListGenerateRow( IMessageLogListingPtr Item, const TSharedRef<STableViewBase>& OwnerTable );
+	
+	/** Callback for changing the selected item in the categories list. */
+	void HandleCategoriesListSelectionChanged( IMessageLogListingPtr Selection, ESelectInfo::Type SelectInfo );
 
-	/** The current log listing widget, if any */
-	TSharedPtr<class SMessageLogListing> LogListing;
+private:	
 
-	/** The combo box for selecting the current listing */
-	TSharedPtr< SComboBox< TSharedPtr<class FMessageLogListingViewModel> > > CurrentListingSelector;
+	/** Holds the log categories list view widget. */
+	TSharedPtr<SListView<IMessageLogListingPtr>> CategoriesListView;
 
-	/** The widget for displaying the current listing */
+	/** The widget for displaying the current listing. */
 	TSharedPtr<SBorder> CurrentListingDisplay;
 
-	/** The widget for displaying the current listing's options, if any */
-	TSharedPtr<SBorder> CurrentListingOptions;
+	/** The current log listing widget, if any. */
+	TSharedPtr<class SMessageLogListing> LogListing;
 
-	/**	Whether the view is currently updating the viewmodel selection */
-	bool bUpdatingSelection;
+	/** The message log view model */
+	TSharedPtr<class FMessageLogViewModel> ViewModel;
 };
-
diff --git a/Engine/Source/Developer/MessageLog/Private/UserInterface/SMessageLogCategoryListRow.h b/Engine/Source/Developer/MessageLog/Private/UserInterface/SMessageLogCategoryListRow.h
new file mode 100644
--- /dev/null
+++ b/Engine/Source/Developer/MessageLog/Private/UserInterface/SMessageLogCategoryListRow.h
@@ -0,0 +1,95 @@
+// Copyright 1998-2014 Epic Games, Inc. All Rights Reserved.
+
+#pragma once
+
+
+#define LOCTEXT_NAMESPACE "SMessageLogCategoryListRow"
+
+
+/**
+ * Implements a row widget for the log categories list view.
+ */
+class SMessageLogCategoryListRow
+	: public SMultiColumnTableRow<IMessageLogListingPtr>
+{
+public:
+
+	SLATE_BEGIN_ARGS(SMessageLogCategoryListRow) { }
+	SLATE_END_ARGS()
+
+public:
+
+	/**
+	 * Constructs the widget.
+	 *
+	 * @param InArgs The arguments.
+	 */
+	void Construct( const FArguments& InArgs, const IMessageLogListingRef& InLogListing, const TSharedRef<STableViewBase>& InOwnerTableView )
+	{
+		LogListing = InLogListing;
+		
+		SMultiColumnTableRow<IMessageLogListingPtr>::Construct(FSuperRowType::FArguments(), InOwnerTableView);
+	}
+
+public:
+
+	/**
+	 * Generates the widget for the specified column.
+	 *
+	 * @param ColumnName The name of the column to generate the widget for.
+	 * @return The widget.
+	 */
+	BEGIN_SLATE_FUNCTION_BUILD_OPTIMIZATION
+	virtual TSharedRef<SWidget> GenerateWidgetForColumn( const FName& ColumnName ) override
+	{
+		if (ColumnName == TEXT("Name"))
+		{
+			return SNew(SBox)
+				.Padding(FMargin(4.0f))
+				.VAlign(VAlign_Center)
+				[
+					SNew(STextBlock)
+						.ColorAndOpacity(this, &SMessageLogCategoryListRow::HandleTextColorAndOpacity)
+						.Text(this, &SMessageLogCategoryListRow::HandleNameColumnText)
+				];
+		}
+
+		return SNullWidget::NullWidget;
+	}
+	END_SLATE_FUNCTION_BUILD_OPTIMIZATION
+
+private:
+
+	// Callback for getting the text in the 'Entries' column.
+	FText HandleNameColumnText( ) const
+	{
+		const FText CategoryName = FText::FromString(LogListing->GetName().ToString());
+		int32 NumMessages = LogListing->GetFilteredMessages().Num();
+
+		if (NumMessages == 0)
+		{
+			return FText::Format(LOCTEXT("ColumnNameFormat", "{0}"), CategoryName);
+		}
+
+		return FText::Format(LOCTEXT("ColumnNameCountFormat", "{0} ({1})"), CategoryName, FText::AsNumber(NumMessages));		
+	}
+
+	// Callback for getting the text color.
+	FSlateColor HandleTextColorAndOpacity( ) const
+	{
+		if (true)
+		{
+			return FSlateColor::UseForeground();
+		}
+
+		return FSlateColor::UseSubduedForeground();
+	}
+
+private:
+
+	// Holds the log listing used to populate this row.
+	IMessageLogListingPtr LogListing;
+};
+
+
+#undef LOCTEXT_NAMESPACE
diff --git a/Engine/Source/Developer/MessageLog/Private/UserInterface/SMessageLogListing.cpp b/Engine/Source/Developer/MessageLog/Private/UserInterface/SMessageLogListing.cpp
--- a/Engine/Source/Developer/MessageLog/Private/UserInterface/SMessageLogListing.cpp
+++ b/Engine/Source/Developer/MessageLog/Private/UserInterface/SMessageLogListing.cpp
@@ -3,246 +3,16 @@
 #include "MessageLogPrivatePCH.h"
 #include "SMessageLogListing.h"
 #include "IMessageLogListing.h"
-#include "UObjectToken.h"
+#include "SMessageLogMessageListRow.h"
 
-#if WITH_EDITOR
-#include "IDocumentation.h"
-#endif
 
 #define LOCTEXT_NAMESPACE "Developer.MessageLog"
 
-class SMessageLogListingItem : public STableRow< TSharedPtr< FTokenizedMessage > >
-{
-public:
-	DECLARE_DELEGATE_TwoParams( FOnTokenClicked, TSharedPtr<FTokenizedMessage>, const TSharedRef<class IMessageToken>& );
-
-public:
-	SLATE_BEGIN_ARGS(SMessageLogListingItem)
-		: _Message()
-		, _OnTokenClicked()
-		{}
-		SLATE_ATTRIBUTE(TSharedPtr<FTokenizedMessage>, Message)
-		SLATE_EVENT( FOnTokenClicked, OnTokenClicked )
-	SLATE_END_ARGS()
-
-	/**
-	 * Construct child widgets that comprise this widget.
-	 *
-	 * @param InArgs  Declaration from which to construct this widget
-	 */
-	void Construct( const FArguments& InArgs, const TSharedRef< STableViewBase >& InOwnerTableView )
-	{
-		this->OnTokenClicked = InArgs._OnTokenClicked;
-
-		Message = InArgs._Message.Get();
-
-		STableRow< TSharedPtr< FTokenizedMessage > >::Construct(
-			STableRow< TSharedPtr< FTokenizedMessage > >::FArguments()
-			.Content()
-			[
-				GenerateWidget()
-			],
-			InOwnerTableView );
-	}
-
-	/** @return Widget for this log listing item*/
-	virtual TSharedRef<SWidget> GenerateWidget()
-	{
-		// See if we have any valid tokens which match the column name
-		const TArray< TSharedRef<IMessageToken> >& MessageTokens = Message->GetMessageTokens();
-
-		// Collect entire message for tooltip
-		FText FullMessage = Message->ToText();
-
-		// Create the horizontal box and add the icon
-		TSharedPtr<SHorizontalBox> HorzBox;
-		TSharedRef<SWidget> CellContent =
-			SNew(SBorder)
-			.OnMouseButtonUp(this, &SMessageLogListingItem::Message_OnMouseButtonUp)
-			.BorderImage(FEditorStyle::GetBrush("NoBorder"))
-			.Padding( FMargin(2.f, 2.f, 2.f, 1.f) )
-			[
-				SAssignNew(HorzBox, SHorizontalBox)
-				.ToolTipText(FullMessage)
-			];
-		check(HorzBox.IsValid());
-
-		// Iterate over parts of the message and create widgets for them
-		for (auto TokenIt = MessageTokens.CreateConstIterator(); TokenIt; ++TokenIt)
-		{
-			const TSharedRef<IMessageToken>& Token = *TokenIt;
-			CreateMessage( HorzBox, Token );
-		}
-
-		return CellContent;
-	}
-
-protected:
-	TSharedPtr<SWidget> CreateHyperlink(const TSharedRef<IMessageToken>& InMessageToken, const FText& InToolTip = FText() )
-	{
-		return SNew(SHyperlink)
-			.Text(InMessageToken->ToText())
-			.ToolTipText(InToolTip)
-			.TextStyle(FEditorStyle::Get(), "MessageLog")
-			.OnNavigate(this, &SMessageLogListingItem::Hyperlink_OnNavigate, InMessageToken);
-	}
-
-	void CreateMessage( TSharedPtr<SHorizontalBox> InHorzBox, const TSharedRef<IMessageToken>& InMessageToken )
-	{
-		if ( InHorzBox.IsValid() )
-		{
-			TSharedPtr<SWidget> Content;
-
-			switch(InMessageToken->GetType())
-			{
-			case EMessageToken::Text:
-			case EMessageToken::String:
-				{
-					if(InMessageToken->GetOnMessageTokenActivated().IsBound())
-					{
-						Content = CreateHyperlink(InMessageToken);
-					}
-					else
-					{
-						SAssignNew(Content, STextBlock)
-							.Text( InMessageToken->ToText() )
-							.TextStyle( FEditorStyle::Get(), TEXT("MessageLog") );
-					}
-				}
-				break;
-			case EMessageToken::Image:
-				{
-					const TSharedRef<FImageToken> ImageToken = StaticCastSharedRef<FImageToken>(InMessageToken);
-					if(ImageToken->GetImageName() != NAME_None)
-					{
-						if(InMessageToken->GetOnMessageTokenActivated().IsBound())
-						{
-							SAssignNew(Content, SButton)
-								.OnClicked(this, &SMessageLogListingItem::Button_OnClicked, InMessageToken)
-								.Content()
-								[
-									SNew(SImage)
-									.Image( FEditorStyle::GetBrush( ImageToken->GetImageName() ) )
-								];
-						}
-						else
-						{
-							SAssignNew(Content, SImage)
-								.Image( FEditorStyle::GetBrush( ImageToken->GetImageName() ) );
-						}
-					}
-				}
-				break;
-			case EMessageToken::Severity:
-				{
-					const TSharedRef<FSeverityToken> SeverityToken = StaticCastSharedRef<FSeverityToken>(InMessageToken);
-					FName SeverityBrush = FTokenizedMessage::GetSeverityIconName(SeverityToken->GetSeverity());
-					if(SeverityBrush != NAME_None)
-					{
-						if(InMessageToken->GetOnMessageTokenActivated().IsBound())
-						{
-							SAssignNew(Content, SButton)
-								.OnClicked(this, &SMessageLogListingItem::Button_OnClicked, InMessageToken)
-								.Content()
-								[
-									SNew(SImage)
-									.Image( FEditorStyle::GetBrush( SeverityBrush ) )
-								];
-						}
-						else
-						{
-							SAssignNew(Content, SImage)
-								.Image( FEditorStyle::GetBrush( SeverityBrush ) );
-						}
-					}
-				}
-				break;
-			case EMessageToken::Object:
-				{
-					const TSharedRef<FUObjectToken> UObjectToken = StaticCastSharedRef<FUObjectToken>(InMessageToken);
-					Content = CreateHyperlink(InMessageToken, FUObjectToken::DefaultOnGetObjectDisplayName().IsBound() ? FUObjectToken::DefaultOnGetObjectDisplayName().Execute(UObjectToken->GetObject().Get(), true) : UObjectToken->ToText());
-				}
-				break;
-			case EMessageToken::URL:
-				{
-					const TSharedRef<FURLToken> URLToken = StaticCastSharedRef<FURLToken>(InMessageToken);
-					Content = CreateHyperlink(InMessageToken, FText::FromString( URLToken->GetURL() ) );
-				}
-				break;
-			case EMessageToken::AssetName:
-				{
-					const TSharedRef<FAssetNameToken> AssetNameToken = StaticCastSharedRef<FAssetNameToken>(InMessageToken);
-					Content = CreateHyperlink(InMessageToken, AssetNameToken->ToText());
-				}
-				break;
-			case EMessageToken::Documentation:
-				{
-#if WITH_EDITOR
-					const TSharedRef<FDocumentationToken> DocumentationToken = StaticCastSharedRef<FDocumentationToken>(InMessageToken);
-					Content = IDocumentation::Get()->CreateAnchor( DocumentationToken->GetDocumentationLink(), DocumentationToken->GetPreviewExcerptLink(), DocumentationToken->GetPreviewExcerptName() );
-#endif
-				}
-				break;
-			}
-
-			AddContentToBox( InHorzBox, Content );
-		}
-	}
-	
-	void AddContentToBox( TSharedPtr<SHorizontalBox> InHorzBox, TSharedPtr<SWidget> InContent ) const
-	{
-		if ( InHorzBox.IsValid() && InContent.IsValid() )
-		{
-			InHorzBox->AddSlot()
-				.AutoWidth()
-				.VAlign(VAlign_Center)
-				.HAlign(HAlign_Left)
-				.Padding(0.f, 0.f, 2.f, 0.f)
-				[
-					InContent.ToSharedRef()
-				];
-		}
-	}
-
-	void Hyperlink_OnNavigate( TSharedRef<IMessageToken> InMessageToken )
-	{
-		InMessageToken->GetOnMessageTokenActivated().ExecuteIfBound( InMessageToken );
-		OnTokenClicked.ExecuteIfBound(Message, InMessageToken);
-	}
-
-	FReply Button_OnClicked( TSharedRef<IMessageToken> InMessageToken )
-	{
-		InMessageToken->GetOnMessageTokenActivated().ExecuteIfBound( InMessageToken );
-		OnTokenClicked.ExecuteIfBound(Message, InMessageToken);
-		return FReply::Handled();
-	}
-
-	FReply Message_OnMouseButtonUp(const FGeometry& MyGeometry, const FPointerEvent& MouseEvent)
-	{
-		if (MouseEvent.GetEffectingButton() == EKeys::RightMouseButton)
-		{
-			Message->GetRightClickedMethod().ExecuteIfBound(Message);
-		}
-
-		return FReply::Handled();
-	}
-
-protected:
-
-	/** The message used to create this widget */
-	TSharedPtr<FTokenizedMessage> Message;
-
-	/** Delegate to execute when the token is clicked */
-	FOnTokenClicked OnTokenClicked;
-};
-
 
 SMessageLogListing::SMessageLogListing()
 	: UICommandList( MakeShareable( new FUICommandList ) )
 	, bUpdatingSelection( false )
-{
-
-}
+{ }
 
 
 SMessageLogListing::~SMessageLogListing()
@@ -251,6 +21,7 @@ SMessageLogListing::~SMessageLogListing()
 	MessageLogListingViewModel->OnSelectionChanged().RemoveAll( this );
 }
 
+
 void SMessageLogListing::Construct( const FArguments& InArgs, const TSharedRef< IMessageLogListing >& InModelView )
 {
 	MessageLogListingViewModel = StaticCastSharedRef<FMessageLogListingViewModel>(InModelView);
@@ -260,16 +31,26 @@ void SMessageLogListing::Construct( const FArguments& InArgs, const TSharedRef<
 
 	ChildSlot
 	[
-		SAssignNew(VerticalBox, SVerticalBox)
-		+SVerticalBox::Slot()
-		.FillHeight(1)
-		[
-			SAssignNew(MessageListView, SListView< TSharedRef<FTokenizedMessage> >)
-				.ListItemsSource(&MessageLogListingViewModel->GetFilteredMessages())
-				.OnGenerateRow(this, &SMessageLogListing::MakeMessageLogListItemWidget)
-				.OnSelectionChanged(this, &SMessageLogListing::OnLineSelectionChanged)
-				.ItemHeight(24)
-		]
+		SNew(SBorder)
+			.BorderImage(FEditorStyle::GetBrush("ToolPanel.GroupBorder"))
+			[
+				SAssignNew(VerticalBox, SVerticalBox)
+
+				+ SVerticalBox::Slot()
+					.FillHeight(1.0f)
+					.Padding(2.0f)
+					[
+						SNew(SBorder)
+							.BorderImage(FEditorStyle::GetBrush("MessageLog.ListBorder"))
+							[
+								SAssignNew(MessageListView, SListView< TSharedRef<FTokenizedMessage> >)
+									.ListItemsSource(&MessageLogListingViewModel->GetFilteredMessages())
+									.OnGenerateRow(this, &SMessageLogListing::MakeMessageLogListItemWidget)
+									.OnSelectionChanged(this, &SMessageLogListing::OnLineSelectionChanged)
+									.ItemHeight(24.0f)
+							]
+					]
+			]
 	];
 
 	//If we have some content below the message log, add a separator and a new box.
@@ -279,13 +60,7 @@ void SMessageLogListing::Construct( const FArguments& InArgs, const TSharedRef<
 	{
 		VerticalBox->AddSlot()
 			.AutoHeight()
-			.Padding(0, 0, 0, 1)
-			[
-				SNew(SSeparator)
-			];
-
-		VerticalBox->AddSlot()
-			.AutoHeight()
+			.Padding(6.0f)
 			[
 				SAssignNew(HorizontalBox, SHorizontalBox)
 			];
@@ -364,6 +139,7 @@ void SMessageLogListing::OnChanged()
 	RefreshVisibility();
 }
 
+
 void SMessageLogListing::OnSelectionChanged()
 {
 	if( bUpdatingSelection )
@@ -386,6 +162,7 @@ void SMessageLogListing::OnSelectionChanged()
 	bUpdatingSelection = false;
 }
 
+
 void SMessageLogListing::RefreshVisibility()
 {
 	const TArray< TSharedRef<FTokenizedMessage> >& Messages = MessageLogListingViewModel->GetFilteredMessages();
@@ -397,6 +174,7 @@ void SMessageLogListing::RefreshVisibility()
 	MessageListView->RequestListRefresh();
 }
 
+
 void SMessageLogListing::BroadcastMessageTokenClicked( TSharedPtr<FTokenizedMessage> Message, const TSharedRef<IMessageToken>& Token )
 {
 	ClearSelectedMessages();
@@ -404,21 +182,25 @@ void SMessageLogListing::BroadcastMessageTokenClicked( TSharedPtr<FTokenizedMess
 	MessageLogListingViewModel->ExecuteToken(Token);
 }
 
+
 const TArray< TSharedRef<FTokenizedMessage> > SMessageLogListing::GetSelectedMessages() const
 {
 	return MessageLogListingViewModel->GetSelectedMessages();
 }
 
+
 void SMessageLogListing::SelectMessage( const TSharedRef<class FTokenizedMessage>& Message, bool bSelected ) const
 {
 	MessageLogListingViewModel->SelectMessage( Message, bSelected );
 }
 
+
 bool SMessageLogListing::IsMessageSelected( const TSharedRef<class FTokenizedMessage>& Message ) const
 {
 	return MessageLogListingViewModel->IsMessageSelected( Message );
 }
 
+
 void SMessageLogListing::ScrollToMessage( const TSharedRef<class FTokenizedMessage>& Message ) const
 {
 	if(!MessageListView->IsItemVisible(Message))
@@ -427,16 +209,19 @@ void SMessageLogListing::ScrollToMessage( const TSharedRef<class FTokenizedMessa
 	}
 }
 
+
 void SMessageLogListing::ClearSelectedMessages() const
 {
 	MessageLogListingViewModel->ClearSelectedMessages();
 }
 
+
 void SMessageLogListing::InvertSelectedMessages() const
 {
 	MessageLogListingViewModel->InvertSelectedMessages();
 }
 
+
 FText SMessageLogListing::GetSelectedMessagesAsText() const
 {
 	return MessageLogListingViewModel->GetSelectedMessagesAsText();
@@ -448,39 +233,37 @@ FText SMessageLogListing::GetAllMessagesAsText() const
 	return MessageLogListingViewModel->GetAllMessagesAsText();
 }
 
+
 TSharedRef<ITableRow> SMessageLogListing::MakeMessageLogListItemWidget( TSharedRef<FTokenizedMessage> Message, const TSharedRef<STableViewBase>& OwnerTable )
 {
 	return
-		SNew(SMessageLogListingItem, OwnerTable)
+		SNew(SMessageLogMessageListRow, OwnerTable)
 		.Message(Message)
 		.OnTokenClicked( this, &SMessageLogListing::BroadcastMessageTokenClicked );
 }
 
+
 void SMessageLogListing::OnLineSelectionChanged( TSharedPtr< FTokenizedMessage > Selection, ESelectInfo::Type /*SelectInfo*/ )
 {
-	if( bUpdatingSelection )
+	if (bUpdatingSelection)
 	{
 		return;
 	}
 
 	bUpdatingSelection = true;
-	TArray< TSharedRef< FTokenizedMessage > > SelectedItems = MessageListView->GetSelectedItems();
-	
-	MessageLogListingViewModel->SelectMessages( SelectedItems );
-	
-	if ( Selection.IsValid() )
 	{
-		Selection->GetSelectionChangedMethod().ExecuteIfBound( SelectedItems );
+		MessageLogListingViewModel->SelectMessages(MessageListView->GetSelectedItems());
 	}
-
 	bUpdatingSelection = false;
 }
 
+
 void SMessageLogListing::CopySelectedToClipboard() const
 {
 	CopyText( true, true );
 }
 
+
 FText SMessageLogListing::CopyText( bool bSelected, bool bClipboard ) const
 {
 	FText CombinedString;
@@ -504,16 +287,19 @@ FText SMessageLogListing::CopyText( bool bSelected, bool bClipboard ) const
 	return CombinedString;
 }
 
+
 const TSharedRef< const FUICommandList > SMessageLogListing::GetCommandList() const 
 { 
 	return UICommandList;
 }
 
+
 FReply SMessageLogListing::OnKeyDown( const FGeometry& MyGeometry, const FKeyboardEvent& InKeyboardEvent )
 {
 	return UICommandList->ProcessCommandBindings( InKeyboardEvent ) ? FReply::Handled() : FReply::Unhandled();
 }
 
+
 EVisibility SMessageLogListing::GetFilterMenuVisibility()
 {
 	if( MessageLogListingViewModel->GetShowFilters() )
@@ -524,6 +310,7 @@ EVisibility SMessageLogListing::GetFilterMenuVisibility()
 	return EVisibility::Hidden;
 }
 
+
 TSharedRef<ITableRow> SMessageLogListing::MakeShowWidget(TSharedRef<FMessageFilter> Selection, const TSharedRef<STableViewBase>& OwnerTable)
 {
 	return
@@ -548,6 +335,7 @@ TSharedRef<ITableRow> SMessageLogListing::MakeShowWidget(TSharedRef<FMessageFilt
 		];
 }
 
+
 TSharedRef<SWidget> SMessageLogListing::OnGetFilterMenuContent()
 {
 	return 
@@ -557,6 +345,7 @@ TSharedRef<SWidget> SMessageLogListing::OnGetFilterMenuContent()
 		.ItemHeight(24);
 }
 
+
 FText SMessageLogListing::OnGetPageMenuLabel() const
 {
 	if(MessageLogListingViewModel->GetPageCount() > 1)
@@ -569,6 +358,7 @@ FText SMessageLogListing::OnGetPageMenuLabel() const
 	}
 }
 
+
 TSharedRef<SWidget> SMessageLogListing::OnGetPageMenuContent() const
 {
 	if(MessageLogListingViewModel->GetPageCount() > 1)
@@ -591,16 +381,19 @@ TSharedRef<SWidget> SMessageLogListing::OnGetPageMenuContent() const
 	return SNullWidget::NullWidget;
 }
 
+
 void SMessageLogListing::OnPageSelected(uint32 PageIndex)
 {
 	MessageLogListingViewModel->SetCurrentPageIndex(PageIndex);
 }
 
+
 bool SMessageLogListing::IsPageWidgetEnabled() const
 {
 	return MessageLogListingViewModel->GetPageCount() > 1;
 }
 
+
 EVisibility SMessageLogListing::GetPageWidgetVisibility() const
 {
 	if(MessageLogListingViewModel->GetShowPages())
@@ -613,11 +406,13 @@ EVisibility SMessageLogListing::GetPageWidgetVisibility() const
 	}
 }
 
+
 bool SMessageLogListing::IsClearWidgetEnabled() const
 {
 	return MessageLogListingViewModel->NumMessages() > 0;
 }
 
+
 EVisibility SMessageLogListing::GetClearWidgetVisibility() const
 {
 	if (MessageLogListingViewModel->GetAllowClear() && !MessageLogListingViewModel->GetShowPages())
@@ -630,10 +425,12 @@ EVisibility SMessageLogListing::GetClearWidgetVisibility() const
 	}
 }
 
+
 FReply SMessageLogListing::OnClear()
 {
 	MessageLogListingViewModel->ClearMessages();
 	return FReply::Handled();
 }
 
+
 #undef LOCTEXT_NAMESPACE
diff --git a/Engine/Source/Developer/MessageLog/Private/UserInterface/SMessageLogMessageListRow.h b/Engine/Source/Developer/MessageLog/Private/UserInterface/SMessageLogMessageListRow.h
new file mode 100644
--- /dev/null
+++ b/Engine/Source/Developer/MessageLog/Private/UserInterface/SMessageLogMessageListRow.h
@@ -0,0 +1,336 @@
+// Copyright 1998-2014 Epic Games, Inc. All Rights Reserved.
+
+#pragma once
+
+#include "UObjectToken.h"
+
+#if WITH_EDITOR
+	#include "IDocumentation.h"
+#endif
+
+#define LOCTEXT_NAMESPACE "SMessageLogMessageListRow"
+
+
+class SMessageLogMessageListRow
+	: public STableRow<TSharedPtr<FTokenizedMessage>>
+{
+public:
+
+	DECLARE_DELEGATE_TwoParams( FOnTokenClicked, TSharedPtr<FTokenizedMessage>, const TSharedRef<class IMessageToken>& );
+
+public:
+
+	SLATE_BEGIN_ARGS(SMessageLogMessageListRow)
+		: _Message()
+		, _OnTokenClicked()
+	{ }
+		SLATE_ATTRIBUTE(TSharedPtr<FTokenizedMessage>, Message)
+		SLATE_EVENT(FOnTokenClicked, OnTokenClicked)
+	SLATE_END_ARGS()
+
+	/**
+	 * Construct child widgets that comprise this widget.
+	 *
+	 * @param InArgs  Declaration from which to construct this widget
+	 */
+	void Construct( const FArguments& InArgs, const TSharedRef< STableViewBase >& InOwnerTableView )
+	{
+		this->OnTokenClicked = InArgs._OnTokenClicked;
+
+		Message = InArgs._Message.Get();
+
+		STableRow<TSharedPtr<FTokenizedMessage>>::Construct(
+			STableRow<TSharedPtr<FTokenizedMessage>>::FArguments()
+			.Content()
+			[
+				GenerateWidget()
+			],
+			InOwnerTableView
+		);
+	}
+
+	/** @return Widget for this log listing item*/
+	virtual TSharedRef<SWidget> GenerateWidget()
+	{
+		// See if we have any valid tokens which match the column name
+		const TArray<TSharedRef<IMessageToken>>& MessageTokens = Message->GetMessageTokens();
+
+		// Create the horizontal box and add the icon
+		TSharedRef<SVerticalBox> DetailBox = SNew(SVerticalBox);
+		TSharedRef<SHorizontalBox> TokenBox = SNew(SHorizontalBox);
+		FName SeverityImageName = NAME_None;
+		bool HasTokens = false;
+		FText TitleText;
+
+		// Iterate over parts of the message and create widgets for them
+		for (auto TokenIt = MessageTokens.CreateConstIterator(); TokenIt; ++TokenIt)
+		{
+			const TSharedRef<IMessageToken>& Token = *TokenIt;
+
+			if (Token->GetType() == EMessageToken::Severity)
+			{
+				if (SeverityImageName == NAME_None)
+				{
+					const TSharedRef<FSeverityToken> SeverityToken = StaticCastSharedRef<FSeverityToken>(Token);
+					SeverityImageName = FTokenizedMessage::GetSeverityIconName(SeverityToken->GetSeverity());
+				}
+			}
+			else if (Token->GetType() == EMessageToken::Text)
+			{
+				if (TitleText.IsEmpty())
+				{
+					TitleText = Token->ToText();
+				}
+				else
+				{
+					DetailBox->AddSlot()
+						.AutoHeight()
+						.Padding(0.0f, 0.0f, 0.0f, 6.0f)
+						[
+							SNew(STextBlock)
+								.ColorAndOpacity(FSlateColor::UseSubduedForeground())
+								.Text(Token->ToText())
+						];
+				}
+			}
+			else
+			{
+				CreateMessage(TokenBox, Token);
+				HasTokens = true;
+			}		
+		}
+
+		return SNew(SHorizontalBox)
+			.ToolTipText(Message->ToText())
+
+		+ SHorizontalBox::Slot()
+			.AutoWidth()
+			[
+				SNew(SBox)
+					.Padding(2.0f)
+					.HAlign(HAlign_Center)
+					.VAlign(VAlign_Top)
+					[
+						(SeverityImageName == NAME_None)
+							? SNullWidget::NullWidget
+							: SNew(SImage)
+								.Image(FEditorStyle::GetBrush(SeverityImageName))
+					]
+			]
+
+		+ SHorizontalBox::Slot()
+			.FillWidth(1.0f)
+			.Padding(2.0f, 0.0f, 0.0f, 0.0f)
+			[
+				SNew(SBox)
+					.Padding(2.0f)
+					[
+						SNew(SVerticalBox)
+
+						+ SVerticalBox::Slot()
+							.AutoHeight()
+							.VAlign(VAlign_Center)
+							[
+								SNew(STextBlock)
+									.Font(FSlateFontInfo(FPaths::EngineContentDir() / TEXT("Slate/Fonts/Roboto-Bold.ttf"), 9))
+									.Text(TitleText)
+							]
+
+						+ SVerticalBox::Slot()
+							.AutoHeight()
+							.Padding(0.0f, 4.0f, 0.0f, 0.0f)
+							[
+								DetailBox
+							]
+
+						+ SVerticalBox::Slot()
+							.AutoHeight()
+							[
+								!HasTokens
+									? SNullWidget::NullWidget
+									: SNew(SSeparator)
+							]
+
+						+ SVerticalBox::Slot()
+							.AutoHeight()
+							.HAlign(HAlign_Right)
+							[
+								!HasTokens
+									? SNullWidget::NullWidget
+									: TokenBox
+							]
+					]
+			];
+	}
+
+protected:
+
+	TSharedRef<SWidget> CreateHyperlink( const TSharedRef<IMessageToken>& InMessageToken, const FText& InToolTip = FText() )
+	{
+		return SNew(SHyperlink)
+			.Text(InMessageToken->ToText())
+			.ToolTipText(InToolTip)
+			.TextStyle(FEditorStyle::Get(), "MessageLog")
+			.OnNavigate(this, &SMessageLogMessageListRow::HandleHyperlinkNavigate, InMessageToken);
+	}
+
+	void CreateMessage( const TSharedRef<SHorizontalBox>& InHorzBox, const TSharedRef<IMessageToken>& InMessageToken )
+	{
+		TSharedPtr<SWidget> Content;
+		FName IconBrushName;
+
+		switch(InMessageToken->GetType())
+		{/*
+		case EMessageToken::Image:
+			{
+				const TSharedRef<FImageToken> ImageToken = StaticCastSharedRef<FImageToken>(InMessageToken);
+				
+				if (ImageToken->GetImageName() != NAME_None)
+				{
+					if (InMessageToken->GetOnMessageTokenActivated().IsBound())
+					{
+						Content = SNew(SButton)
+							.OnClicked(this, &SMessageLogMessageListRow::HandleTokenButtonClicked, InMessageToken)
+							.Content()
+							[
+								SNew(SImage)
+									.Image(FEditorStyle::GetBrush(ImageToken->GetImageName()))
+							];
+					}
+					else
+					{
+						Content = SNew(SImage)
+							.Image(FEditorStyle::GetBrush(ImageToken->GetImageName()));
+					}
+				}
+			}
+			break;
+*/
+		case EMessageToken::Object:
+			{
+				const TSharedRef<FUObjectToken> UObjectToken = StaticCastSharedRef<FUObjectToken>(InMessageToken);
+				
+				IconBrushName = FName("PropertyWindow.Button_Browse");
+				Content = CreateHyperlink(InMessageToken, FUObjectToken::DefaultOnGetObjectDisplayName().IsBound()
+					? FUObjectToken::DefaultOnGetObjectDisplayName().Execute(UObjectToken->GetObject().Get(), true)
+					: UObjectToken->ToText());
+			}
+			break;
+
+		case EMessageToken::URL:
+			{
+				const TSharedRef<FURLToken> URLToken = StaticCastSharedRef<FURLToken>(InMessageToken);
+
+				IconBrushName = FName("MessageLog.Url");
+				Content = CreateHyperlink(InMessageToken, FText::FromString(URLToken->GetURL()));
+			}
+			break;
+
+		case EMessageToken::Action:
+			{
+				const TSharedRef<FActionToken> ActionToken = StaticCastSharedRef<FActionToken>(InMessageToken);
+
+				IconBrushName = FName("MessageLog.Action");
+				Content = SNew(SHyperlink)
+					.Text(InMessageToken->ToText())
+					.ToolTipText(ActionToken->GetActionDescription())
+					.TextStyle(FEditorStyle::Get(), "MessageLog")
+					.OnNavigate(this, &SMessageLogMessageListRow::HandleActionHyperlinkNavigate, ActionToken);
+			}
+			break;
+
+		case EMessageToken::AssetName:
+			{
+				const TSharedRef<FAssetNameToken> AssetNameToken = StaticCastSharedRef<FAssetNameToken>(InMessageToken);
+
+				IconBrushName = FName("PropertyWindow.Button_Browse");
+				Content = CreateHyperlink(InMessageToken, AssetNameToken->ToText());
+			}
+			break;
+
+		case EMessageToken::Documentation:
+			{
+#if WITH_EDITOR
+				const TSharedRef<FDocumentationToken> DocumentationToken = StaticCastSharedRef<FDocumentationToken>(InMessageToken);
+
+				IconBrushName = FName("MessageLog.Docs");
+				Content = SNew(SHyperlink)
+					.Text(LOCTEXT("DocsLabel", "Docs"))
+					.ToolTip(IDocumentation::Get()->CreateToolTip(
+						LOCTEXT("DocumentationTokenToolTip", "Click to open documentation"),
+						NULL,
+						DocumentationToken->GetPreviewExcerptLink(),
+						DocumentationToken->GetPreviewExcerptName())
+					)
+					.TextStyle(FEditorStyle::Get(), "MessageLog")
+					.OnNavigate(this, &SMessageLogMessageListRow::HandleDocsHyperlinkNavigate, DocumentationToken->GetDocumentationLink());
+#endif
+			}
+			break;
+		}
+
+		if (Content.IsValid() && (IconBrushName != NAME_None))
+		{
+			InHorzBox->AddSlot()
+				.AutoWidth()
+				.VAlign(VAlign_Center)
+				.Padding(16.f, 4.0f, 0.0f, 12.0f)
+				[
+					SNew(SHorizontalBox)
+
+					+ SHorizontalBox::Slot()
+						.AutoWidth()
+						.VAlign(VAlign_Center)
+						[
+							SNew(SImage)
+								.ColorAndOpacity(FSlateColor::UseForeground())
+								.Image(FEditorStyle::GetBrush(IconBrushName))
+						]
+
+					+ SHorizontalBox::Slot()
+						.AutoWidth()
+						.VAlign(VAlign_Center)
+						.Padding(2.0f, 0.0f, 0.0f, 0.0f)
+						[
+							Content.ToSharedRef()
+						]
+				];
+		}
+	}
+
+private:
+
+	void HandleActionHyperlinkNavigate( TSharedRef<FActionToken> ActionToken )
+	{
+		ActionToken->ExecuteAction();
+	}
+
+	void HandleDocsHyperlinkNavigate( FString DocumentationLink )
+	{
+		IDocumentation::Get()->Open(DocumentationLink);
+	}
+
+	void HandleHyperlinkNavigate( TSharedRef<IMessageToken> InMessageToken )
+	{
+		InMessageToken->GetOnMessageTokenActivated().ExecuteIfBound(InMessageToken);
+		OnTokenClicked.ExecuteIfBound(Message, InMessageToken);
+	}
+
+	FReply HandleTokenButtonClicked( TSharedRef<IMessageToken> InMessageToken )
+	{
+		InMessageToken->GetOnMessageTokenActivated().ExecuteIfBound(InMessageToken);
+		OnTokenClicked.ExecuteIfBound(Message, InMessageToken);
+		return FReply::Handled();
+	}
+
+protected:
+
+	/** The message used to create this widget. */
+	TSharedPtr<FTokenizedMessage> Message;
+
+	/** Delegate to execute when the token is clicked. */
+	FOnTokenClicked OnTokenClicked;
+};
+
+
+#undef LOCTEXT_NAMESPACE
diff --git a/Engine/Source/Developer/MessageLog/Public/IMessageLogListing.h b/Engine/Source/Developer/MessageLog/Public/IMessageLogListing.h
--- a/Engine/Source/Developer/MessageLog/Public/IMessageLogListing.h
+++ b/Engine/Source/Developer/MessageLog/Public/IMessageLogListing.h
@@ -5,6 +5,14 @@
 #include "TokenizedMessage.h"
 #include "IMessageLog.h"
 
+
+/** Type definition for shared pointers to instances of IMessageLogListing. */
+typedef TSharedPtr<class IMessageLogListing> IMessageLogListingPtr;
+
+/** Type definition for shared references to instances of IMessageLogListing. */
+typedef TSharedRef<class IMessageLogListing> IMessageLogListingRef;
+
+
 /**
  * A message log listing, such as the Compiler Log, or the Map Check Log.
  */
diff --git a/Engine/Source/Editor/EditorStyle/Private/SlateEditorStyle.cpp b/Engine/Source/Editor/EditorStyle/Private/SlateEditorStyle.cpp
--- a/Engine/Source/Editor/EditorStyle/Private/SlateEditorStyle.cpp
+++ b/Engine/Source/Editor/EditorStyle/Private/SlateEditorStyle.cpp
@@ -816,11 +816,16 @@ void FSlateEditorStyle::FStyle::SetupGeneralStyles()
 			.SetFont( TTF_CORE_FONT( "Fonts/Roboto-Regular", 8 ) )
 			.SetShadowOffset( FVector2D::ZeroVector )
 		);
+
+		Set( "MessageLog.Action", new IMAGE_BRUSH( "Icons/icon_file_choosepackages_16px", Icon16x16) );
+		Set( "MessageLog.Docs", new IMAGE_BRUSH( "Icons/icon_Docs_16x", Icon16x16) );
 		Set( "MessageLog.Error", new IMAGE_BRUSH( "Old/Kismet2/Log_Error", Icon16x16 ) );
 		Set( "MessageLog.Warning", new IMAGE_BRUSH( "Old/Kismet2/Log_Warning", Icon16x16 ) );
 		Set( "MessageLog.Note", new IMAGE_BRUSH( "Old/Kismet2/Log_Note", Icon16x16 ) );
+		Set( "MessageLog.Url", new IMAGE_BRUSH( "Icons/icon_world_16x", Icon16x16 ) );
 
 		Set( "MessageLog.TabIcon", new IMAGE_BRUSH( "Icons/icon_tab_MessageLog_16x", Icon16x16 ) );
+		Set( "MessageLog.ListBorder", new BOX_BRUSH( "/Docking/AppTabContentArea", FMargin(4/16.0f) ) );
 	}
 
 #if WITH_EDITOR || IS_PROGRAM
diff --git a/Engine/Source/Editor/Kismet/Private/BlueprintEditor.cpp b/Engine/Source/Editor/Kismet/Private/BlueprintEditor.cpp
--- a/Engine/Source/Editor/Kismet/Private/BlueprintEditor.cpp
+++ b/Engine/Source/Editor/Kismet/Private/BlueprintEditor.cpp
@@ -1697,14 +1697,7 @@ void FBlueprintEditor::CreateDefaultTabContents(const TArray<UBlueprint*>& InBlu
 	CompilerResultsListing = MessageLogModule.CreateLogListing( "BlueprintEditorCompileResults" );
 	CompilerResultsListing->OnMessageTokenClicked().AddSP(this, &FBlueprintEditor::OnLogTokenClicked);
 
-	CompilerResults = 
-		SNew(SBorder)
-		.Padding(0.0f)
-		.BorderImage(FEditorStyle::GetBrush("ToolPanel.GroupBorder"))
-		[
-			MessageLogModule.CreateLogListingWidget( CompilerResultsListing.ToSharedRef() )
-		];
-
+	CompilerResults = MessageLogModule.CreateLogListingWidget( CompilerResultsListing.ToSharedRef() );
 	FindResults = SNew(SFindInBlueprints, SharedThis(this));
 	
 	this->Inspector = 
diff --git a/Engine/Source/Editor/MaterialEditor/Private/MaterialEditor.cpp b/Engine/Source/Editor/MaterialEditor/Private/MaterialEditor.cpp
--- a/Engine/Source/Editor/MaterialEditor/Private/MaterialEditor.cpp
+++ b/Engine/Source/Editor/MaterialEditor/Private/MaterialEditor.cpp
@@ -613,13 +613,7 @@ void FMaterialEditor::CreateInternalWidgets()
 	LogOptions.MaxPageCount = 1;
 	StatsListing = MessageLogModule.CreateLogListing( "MaterialEditorStats", LogOptions );
 
-	Stats = 
-		SNew(SBorder)
-		.Padding(0.0f)
-		.BorderImage(FEditorStyle::GetBrush("ToolPanel.GroupBorder"))
-		[
-			MessageLogModule.CreateLogListingWidget( StatsListing.ToSharedRef() )
-		];
+	Stats = MessageLogModule.CreateLogListingWidget( StatsListing.ToSharedRef() );
 
 	FindResults =
 		SNew(SFindInMaterial, SharedThis(this));
diff --git a/Engine/Source/Editor/UnrealEd/Private/EditorServer.cpp b/Engine/Source/Editor/UnrealEd/Private/EditorServer.cpp
--- a/Engine/Source/Editor/UnrealEd/Private/EditorServer.cpp
+++ b/Engine/Source/Editor/UnrealEd/Private/EditorServer.cpp
@@ -3467,7 +3467,8 @@ bool UEditorEngine::Map_Check( UWorld* InWorld, const TCHAR* Str, FOutputDevice&
 		Arguments.Add(TEXT("Warnings"), WarningCount - ErrorCount);
 		Arguments.Add(TEXT("Time"), CurrentTime);
 		FMessageLog("MapCheck").Info()
-			->AddToken(FTextToken::Create(FText::Format( LOCTEXT( "MapCheck_Complete", "========== Map Check: {Errors} Error(s), {Warnings} Warning(s), [{Time}ms] ==========" ), Arguments ) ));
+			->AddToken(FTextToken::Create(LOCTEXT( "MapCheck_Complete", "Map check completed" )))
+			->AddToken(FTextToken::Create(FText::Format( LOCTEXT( "MapCheck_CompleteDetail", "There are {Errors} Error(s) and {Warnings} Warning(s), and it took {Time}ms to complete." ), Arguments ) ));
 	}
 
 	GWarn->EndSlowTask();
diff --git a/Engine/Source/Editor/UnrealEd/Private/Kismet2/CompilerResultsLog.cpp b/Engine/Source/Editor/UnrealEd/Private/Kismet2/CompilerResultsLog.cpp
--- a/Engine/Source/Editor/UnrealEd/Private/Kismet2/CompilerResultsLog.cpp
+++ b/Engine/Source/Editor/UnrealEd/Private/Kismet2/CompilerResultsLog.cpp
@@ -178,9 +178,7 @@ void FCompilerResultsLog::InternalLogMessage(const EMessageSeverity::Type& Sever
 				ObjText = LOCTEXT("None", "(none)");
 			}
 			
-			TSharedRef<FUObjectToken> Token = FUObjectToken::Create( ObjectArgument, ObjText );
-			Token->OnMessageTokenActivated(FOnMessageTokenActivated::CreateRaw(this, &FCompilerResultsLog::OnTokenActivated));
-			Line->AddToken( Token );
+			Line->AddToken( FUObjectToken::Create( ObjectArgument, ObjText ) );
 		}
 		Line->AddToken( FTextToken::Create( FText::FromString( Start ) ) );
 	}
@@ -350,11 +348,6 @@ void FCompilerResultsLog::GetGlobalModuleCompilerDump(const FString& LogDump, EC
 	MessageLog.AddMessages(ParseCompilerLogDump(LogDump));
 }
 
-void FCompilerResultsLog::OnTokenActivated(const TSharedRef<class IMessageToken>& InTokenRef)
-{
-
-}
-
 // Note: Message is not a fprintf string!  It should be preformatted, but can contain @@ to indicate object references, which are the varargs
 void FCompilerResultsLog::Error(const TCHAR* Message, ...)
 {
diff --git a/Engine/Source/Editor/UnrealEd/Public/Kismet2/CompilerResultsLog.h b/Engine/Source/Editor/UnrealEd/Public/Kismet2/CompilerResultsLog.h
--- a/Engine/Source/Editor/UnrealEd/Public/Kismet2/CompilerResultsLog.h
+++ b/Engine/Source/Editor/UnrealEd/Public/Kismet2/CompilerResultsLog.h
@@ -95,9 +95,6 @@ public:
 protected:
 	/** Create a tokenized message record from a message containing @@ indicating where each UObject* in the ArgPtr list goes and place it in the MessageLog. */
 	void InternalLogMessage(const EMessageSeverity::Type& Severity, const TCHAR* Message, va_list ArgPtr);
-
-	/** Callback when a token is activated */
-	void OnTokenActivated(const TSharedRef<class IMessageToken>& InTokenRef);
 	
 	/** */
 	void AnnotateNode(class UEdGraphNode* Node, TSharedRef<FTokenizedMessage> LogLine);
diff --git a/Engine/Source/Runtime/Core/Private/Misc/MessageLog.cpp b/Engine/Source/Runtime/Core/Private/Misc/MessageLog.cpp
--- a/Engine/Source/Runtime/Core/Private/Misc/MessageLog.cpp
+++ b/Engine/Source/Runtime/Core/Private/Misc/MessageLog.cpp
@@ -93,23 +93,13 @@ FMessageLog::~FMessageLog()
 
 const TSharedRef<FTokenizedMessage>& FMessageLog::AddMessage( const TSharedRef<FTokenizedMessage>& InMessage )
 {
-	if(MessageSelectionChanged.IsBound())
-	{
-		InMessage->SetSelectionChangedMethod( MessageSelectionChanged );
-	}
 	Messages.Add(InMessage);
+
 	return InMessage;
 }
 
 void FMessageLog::AddMessages( const TArray< TSharedRef<FTokenizedMessage> >& InMessages )
 {
-	if(MessageSelectionChanged.IsBound())
-	{
-		for(TArray< TSharedRef<FTokenizedMessage> >::TConstIterator It = InMessages.CreateConstIterator(); It; It++)
-		{
-			(*It)->SetSelectionChangedMethod( MessageSelectionChanged );
-		}
-	}
 	Messages.Append(InMessages);
 }
 
diff --git a/Engine/Source/Runtime/Core/Private/Misc/TokenizedMessage.cpp b/Engine/Source/Runtime/Core/Private/Misc/TokenizedMessage.cpp
--- a/Engine/Source/Runtime/Core/Private/Misc/TokenizedMessage.cpp
+++ b/Engine/Source/Runtime/Core/Private/Misc/TokenizedMessage.cpp
@@ -129,26 +129,6 @@ const TArray<TSharedRef<IMessageToken> >& FTokenizedMessage::GetMessageTokens()
 	return MessageTokens;
 }
 
-void FTokenizedMessage::SetRightClickedMethod(const FOnMessageRightClicked& InRightClickedMethod)
-{
-	RightClickedMethod = InRightClickedMethod;
-}
-
-const FOnMessageRightClicked& FTokenizedMessage::GetRightClickedMethod() const
-{
-	return RightClickedMethod;
-}
-
-void FTokenizedMessage::SetSelectionChangedMethod(const FOnMessageSelectionChanged& InMessageSelectionChanged)
-{
-	MessageSelectionChangedMethod = InMessageSelectionChanged;
-}
-
-const FOnMessageSelectionChanged& FTokenizedMessage::GetSelectionChangedMethod() const
-{
-	return MessageSelectionChangedMethod;
-}
-
 FURLToken::FGenerateURL FURLToken::GenerateURL;
 
 void FURLToken::VisitURL(const TSharedRef<IMessageToken>& Token, FString InURL)
diff --git a/Engine/Source/Runtime/Core/Public/Misc/TokenizedMessage.h b/Engine/Source/Runtime/Core/Public/Misc/TokenizedMessage.h
--- a/Engine/Source/Runtime/Core/Public/Misc/TokenizedMessage.h
+++ b/Engine/Source/Runtime/Core/Public/Misc/TokenizedMessage.h
@@ -26,7 +26,7 @@ namespace EMessageToken
 {
 	enum Type
 	{
-		String,
+		Action,
 		Text,
 		Image,
 		Severity,
@@ -93,11 +93,6 @@ protected:
 	FText CachedText;
 };
 
-/** Delegate used when selecting Message lines */
-DECLARE_DELEGATE_OneParam(FOnMessageSelectionChanged, TArray< TSharedRef<class FTokenizedMessage> >&);
-
-/** Delegate used when right clicking a message line */
-DECLARE_DELEGATE_OneParam(FOnMessageRightClicked, TSharedPtr<class FTokenizedMessage>&);
 
 /** Dummy struct that other structs need to inherit from if they wish to assign to MessageData */
 struct FTokenizedMiscData {};
@@ -168,34 +163,6 @@ public:
 	 */
 	CORE_API const TArray< TSharedRef<IMessageToken> >& GetMessageTokens() const;
 
-	/** 
-	 * Set the delegate for when a message is right-clicked
-	 * 
-	 * @param	InRightClickedMethod	The delegate to call when a message is right clicked
-	 */
-	CORE_API void SetRightClickedMethod(const FOnMessageRightClicked& InRightClickedMethod);
-
-	/** 
-	 * Get the delegate for when a message is right-clicked
-	 * 
-	 * @returns the delegate to call when a message is right clicked
-	 */	
-	CORE_API const FOnMessageRightClicked& GetRightClickedMethod() const;
-
-	/** 
-	 * Set the delegate for when a messages selection has changed
-	 * 
-	 * @param	InMessageSelectionChanged	The delegate to call when a messages selection has changed
-	 */
-	CORE_API void SetSelectionChangedMethod(const FOnMessageSelectionChanged& InMessageSelectionChanged);
-
-	/** 
-	 * Get the delegate for when a messages selection has changed
-	 * 
-	 * @returns the delegate to call when a messages selection has changed
-	 */	
-	CORE_API const FOnMessageSelectionChanged& GetSelectionChangedMethod() const;
-
 	/** 
 	 * Helper function for getting a severity as text
 	 *
@@ -216,25 +183,21 @@ private:
 	/** Private constructor - we want to only create these structures as shared references via Create() */
 	FTokenizedMessage()
 		: Severity( EMessageSeverity::Info )
-		, MessageData( NULL )
-	{}
+		, MessageData( nullptr )
+	{ }
 
 protected:
+
 	/** the array of message tokens this message contains */
 	TArray< TSharedRef<IMessageToken> > MessageTokens;
 
 private:
+
 	/** The severity of this message */
 	EMessageSeverity::Type Severity;
 
 	/** A payload for this message */
 	FTokenizedMiscData* MessageData;	// @todo Don't use TSharedPtr as it'll get free'd mid-sort! (refcount == 0)
-
-	/** Delegate for when a message is right-clicked */
-	FOnMessageRightClicked RightClickedMethod;
-
-	/** Delegate for when message selection is changed */
-	FOnMessageSelectionChanged MessageSelectionChangedMethod;
 };
 
 /** Basic message token with a localized text payload */
@@ -479,4 +442,60 @@ private:
 
 	/** The excerpt to display */
 	FString PreviewExcerptName;
+};
+
+
+DECLARE_DELEGATE(FOnActionTokenExecuted);
+
+/**
+ * Message token that performs an action when activated.
+ */
+class FActionToken
+	: public IMessageToken
+{
+	/** Factory method, tokens can only be constructed as shared refs */
+	CORE_API static TSharedRef<FActionToken> Create( const FText& ActionName, const FText& ActionDescription, const FOnActionTokenExecuted& Action )
+	{
+		return MakeShareable(new FActionToken(ActionName, ActionDescription, Action));
+	}
+
+public:
+
+	/** Executes the assigned action delegate. */
+	void ExecuteAction( )
+	{
+		Action.ExecuteIfBound();
+	}
+
+	/** Gets the action's description text. */
+	const FText& GetActionDescription( )
+	{
+		return ActionDescription;
+	}
+
+public:
+
+	// IMessageToken interface
+	virtual EMessageToken::Type GetType() const override
+	{
+		return EMessageToken::Action;
+	}
+
+protected:
+
+	/** Hidden constructor. */
+	FActionToken( const FText& InActionName, const FText& InActionDescription, const FOnActionTokenExecuted& InAction )
+		: Action(InAction)
+		, ActionDescription(InActionDescription)
+	{
+		CachedText = InActionName;
+	}
+
+private:
+
+	/** Holds a delegate that is executed when this token is activated. */
+	FOnActionTokenExecuted Action;
+
+	/** The action's description text. */
+	const FText ActionDescription;
 };
\ No newline at end of file
diff --git a/Engine/Source/Runtime/Engine/Private/PrimitiveComponent.cpp b/Engine/Source/Runtime/Engine/Private/PrimitiveComponent.cpp
--- a/Engine/Source/Runtime/Engine/Private/PrimitiveComponent.cpp
+++ b/Engine/Source/Runtime/Engine/Private/PrimitiveComponent.cpp
@@ -571,7 +571,8 @@ void UPrimitiveComponent::CheckForErrors()
 	{
 		FMessageLog("MapCheck").PerformanceWarning()
 			->AddToken(FUObjectToken::Create(Owner))
-			->AddToken(FTextToken::Create(LOCTEXT( "MapCheck_Message_ShadowCasterUsingBoundsScale", "Actor casts dynamic shadows and has a BoundsScale greater than 1! This will have a large performance hit" )))
+			->AddToken(FTextToken::Create(LOCTEXT( "MapCheck_Message_ShadowCasterUsingBoundsScale", "Actor casts dynamic shadows and has a BoundsScale greater than 1" )))
+			->AddToken(FTextToken::Create(LOCTEXT( "MapCheck_Message_ShadowCasterUsingBoundsScaleDetail", "This will have a large performance hit." )))
 			->AddToken(FMapErrorToken::Create(FMapErrors::ShadowCasterUsingBoundsScale));
 	}
 
@@ -579,7 +580,9 @@ void UPrimitiveComponent::CheckForErrors()
 	{
 		FMessageLog("MapCheck").Error()
 			->AddToken(FUObjectToken::Create(Owner))
-			->AddToken(FTextToken::Create(LOCTEXT( "MapCheck_Message_InvalidLightmapSettings", "Component is a static type but has invalid lightmap settings!  Indirect lighting will be black.  Common causes are lightmap resolution of 0, LightmapCoordinateIndex out of bounds." )))
+			->AddToken(FTextToken::Create(LOCTEXT( "MapCheck_Message_InvalidLightmapSettings", "Component is a static type but has invalid lightmap settings" )))
+			->AddToken(FTextToken::Create(LOCTEXT( "MapCheck_Message_InvalidLightmapSettingsDetail", "The indirect lighting on this component will be black." )))
+			->AddToken(FTextToken::Create(LOCTEXT( "MapCheck_Message_InvalidLightmapSettingsDetail2", "Common causes: lightmap resolution of 0, LightmapCoordinateIndex out of bounds" )))
 			->AddToken(FMapErrorToken::Create(FMapErrors::StaticComponentHasInvalidLightmapSettings));
 	}
 }
diff --git a/Engine/Source/Runtime/Engine/Private/StaticMeshActor.cpp b/Engine/Source/Runtime/Engine/Private/StaticMeshActor.cpp
--- a/Engine/Source/Runtime/Engine/Private/StaticMeshActor.cpp
+++ b/Engine/Source/Runtime/Engine/Private/StaticMeshActor.cpp
@@ -126,20 +126,17 @@ void AStaticMeshActor::CheckForErrors()
 
 	if( !StaticMeshComponent.IsValid() )
 	{
-		FFormatNamedArguments Arguments;
-		Arguments.Add(TEXT("ActorName"), FText::FromString(GetName()));
 		MapCheck.Warning()
 			->AddToken(FUObjectToken::Create(this))
-			->AddToken(FTextToken::Create(FText::Format( LOCTEXT( "MapCheck_Message_StaticMeshComponent", "{ActorName}: Static mesh actor has NULL StaticMeshComponent property - please delete" ), Arguments ) ))
+			->AddToken(FTextToken::Create(LOCTEXT( "MapCheck_Message_StaticMeshComponent", "Static mesh actor has NULL StaticMeshComponent property" )))
+			->AddToken(FTextToken::Create(LOCTEXT( "MapCheck_Message_StaticMeshComponentDetail", "Please delete the component or assign a valid one." )))
 			->AddToken(FMapErrorToken::Create(FMapErrors::StaticMeshComponent));
 	}
 	else if( StaticMeshComponent->StaticMesh == NULL )
 	{
-		FFormatNamedArguments Arguments;
-		Arguments.Add(TEXT("ActorName"), FText::FromString(GetName()));
 		MapCheck.Warning()
 			->AddToken(FUObjectToken::Create(this))
-			->AddToken(FTextToken::Create(FText::Format( LOCTEXT( "MapCheck_Message_StaticMeshNull", "{ActorName} : Static mesh actor has NULL StaticMesh property" ), Arguments ) ))
+			->AddToken(FTextToken::Create(LOCTEXT( "MapCheck_Message_StaticMeshNull", "Static mesh actor has NULL StaticMesh property" )))
 			->AddToken(FMapErrorToken::Create(FMapErrors::StaticMeshNull));
 	}
 	else
@@ -164,7 +161,8 @@ void AStaticMeshActor::CheckForErrors()
 					Arguments.Add(TEXT("ActorName1"), FText::FromString(A->GetName()));
 					MapCheck.Warning()
 						->AddToken(FUObjectToken::Create(this))
-						->AddToken(FTextToken::Create(FText::Format( LOCTEXT( "MapCheck_Message_SameLocation", "{ActorName0} is in the same location as {ActorName1}" ), Arguments ) ))
+						->AddToken(FTextToken::Create(LOCTEXT( "MapCheck_Message_SameLocation", "Overlapping actors" )))
+						->AddToken(FTextToken::Create(FText::Format( LOCTEXT( "MapCheck_Message_SameLocationDetail", "{ActorName0} is in the same location as {ActorName1}." ), Arguments ) ))
 						->AddToken(FMapErrorToken::Create(FMapErrors::SameLocation));
 				}
 			}
diff --git a/Engine/Source/Runtime/Engine/Private/StaticMeshComponent.cpp b/Engine/Source/Runtime/Engine/Private/StaticMeshComponent.cpp
--- a/Engine/Source/Runtime/Engine/Private/StaticMeshComponent.cpp
+++ b/Engine/Source/Runtime/Engine/Private/StaticMeshComponent.cpp
@@ -350,11 +350,9 @@ void UStaticMeshComponent::CheckForErrors()
 
 	if (!StaticMesh && (!Owner || !Owner->IsA(AWorldSettings::StaticClass())))	// Ignore worldsettings
 	{
-		FFormatNamedArguments Arguments;
-		Arguments.Add(TEXT("OwnerName"), FText::FromString(OwnerName));
 		FMessageLog("MapCheck").Warning()
 			->AddToken(FUObjectToken::Create(Owner))
-			->AddToken(FTextToken::Create(FText::Format(LOCTEXT( "MapCheck_Message_StaticMeshNull", "{OwnerName} : Static mesh actor has NULL StaticMesh property" ), Arguments ) ))
+			->AddToken(FTextToken::Create(LOCTEXT( "MapCheck_Message_StaticMeshNull", "Static mesh actor has NULL StaticMesh property" )))
 			->AddToken(FMapErrorToken::Create(FMapErrors::StaticMeshNull));
 	}
 
@@ -380,11 +378,9 @@ void UStaticMeshComponent::CheckForErrors()
 
 	if ( BodyInstance.bSimulatePhysics && StaticMesh != NULL && StaticMesh->BodySetup != NULL && StaticMesh->BodySetup->AggGeom.GetElementCount() == 0) 
 	{
-		FFormatNamedArguments Arguments;
-		Arguments.Add(TEXT("ActorName"), FText::FromString(GetName()));
 		FMessageLog("MapCheck").Warning()
 			->AddToken(FUObjectToken::Create(this))
-			->AddToken(FTextToken::Create(FText::Format(LOCTEXT( "MapCheck_Message_SimulatePhyNoSimpleCollision", "{ActorName} : Using bSimulatePhysics but StaticMesh has not simple collision."), Arguments ) ));
+			->AddToken(FTextToken::Create(LOCTEXT( "MapCheck_Message_SimulatePhyNoSimpleCollision", "Using bSimulatePhysics but StaticMesh has not simple collision.")));
 	}
 
 	if( Mobility == EComponentMobility::Movable &&
