MessageLog: Handling log category deslection properly

diff --git a/Engine/Source/Developer/MessageLog/Private/UserInterface/SMessageLog.cpp b/Engine/Source/Developer/MessageLog/Private/UserInterface/SMessageLog.cpp
--- a/Engine/Source/Developer/MessageLog/Private/UserInterface/SMessageLog.cpp
+++ b/Engine/Source/Developer/MessageLog/Private/UserInterface/SMessageLog.cpp
@@ -122,7 +122,10 @@ TSharedRef<ITableRow> SMessageLog::HandleCategoriesListGenerateRow( IMessageLogL
 	
 void SMessageLog::HandleCategoriesListSelectionChanged( IMessageLogListingPtr Selection, ESelectInfo::Type SelectInfo )
 {
-	CurrentListingDisplay->SetContent(SAssignNew(LogListing, SMessageLogListing, Selection.ToSharedRef()));
+	CurrentListingDisplay->SetContent(Selection.IsValid()
+		? SAssignNew(LogListing, SMessageLogListing, Selection.ToSharedRef())
+		: SNullWidget::NullWidget
+	);
 }
 
 
