fixed assignment order

diff --git a/Engine/Source/Developer/MessageLog/Private/UserInterface/SMessageLog.cpp b/Engine/Source/Developer/MessageLog/Private/UserInterface/SMessageLog.cpp
--- a/Engine/Source/Developer/MessageLog/Private/UserInterface/SMessageLog.cpp
+++ b/Engine/Source/Developer/MessageLog/Private/UserInterface/SMessageLog.cpp
@@ -122,9 +122,9 @@ TSharedRef<ITableRow> SMessageLog::HandleCategoriesListGenerateRow( IMessageLogL
 	
 void SMessageLog::HandleCategoriesListSelectionChanged( IMessageLogListingPtr Selection, ESelectInfo::Type SelectInfo )
 {
-	CurrentListingDisplay->SetContent(Selection.IsValid()
-		? SAssignNew(LogListing, SMessageLogListing, Selection.ToSharedRef())
-		: SNullWidget::NullWidget
+	CurrentListingDisplay->SetContent(!Selection.IsValid()
+		? SNullWidget::NullWidget
+		: SAssignNew(LogListing, SMessageLogListing, Selection.ToSharedRef())
 	);
 }
 
