clang fix

diff --git a/Engine/Source/Developer/MessageLog/Private/UserInterface/SMessageLog.cpp b/Engine/Source/Developer/MessageLog/Private/UserInterface/SMessageLog.cpp
--- a/Engine/Source/Developer/MessageLog/Private/UserInterface/SMessageLog.cpp
+++ b/Engine/Source/Developer/MessageLog/Private/UserInterface/SMessageLog.cpp
@@ -122,10 +122,14 @@ TSharedRef<ITableRow> SMessageLog::HandleCategoriesListGenerateRow( IMessageLogL
 	
 void SMessageLog::HandleCategoriesListSelectionChanged( IMessageLogListingPtr Selection, ESelectInfo::Type SelectInfo )
 {
-	CurrentListingDisplay->SetContent(!Selection.IsValid()
-		? SNullWidget::NullWidget
-		: SAssignNew(LogListing, SMessageLogListing, Selection.ToSharedRef())
-	);
+	if (Selection.IsValid())
+	{
+		CurrentListingDisplay->SetContent(SAssignNew(LogListing, SMessageLogListing, Selection.ToSharedRef()));
+	}
+	else
+	{
+		CurrentListingDisplay->SetContent(SNullWidget::NullWidget);
+	}
 }
 
 
