fixed tools build

diff --git a/Engine/Source/Developer/MessageLog/Private/UserInterface/SMessageLogMessageListRow.h b/Engine/Source/Developer/MessageLog/Private/UserInterface/SMessageLogMessageListRow.h
--- a/Engine/Source/Developer/MessageLog/Private/UserInterface/SMessageLogMessageListRow.h
+++ b/Engine/Source/Developer/MessageLog/Private/UserInterface/SMessageLogMessageListRow.h
@@ -305,10 +305,12 @@ private:
 		ActionToken->ExecuteAction();
 	}
 
+#if WITH_EDITOR
 	void HandleDocsHyperlinkNavigate( FString DocumentationLink )
 	{
 		IDocumentation::Get()->Open(DocumentationLink);
 	}
+#endif
 
 	void HandleHyperlinkNavigate( TSharedRef<IMessageToken> InMessageToken )
 	{
