MessageLog: Message token layout tweak

diff --git a/Engine/Source/Developer/MessageLog/Private/UserInterface/SMessageLogMessageListRow.h b/Engine/Source/Developer/MessageLog/Private/UserInterface/SMessageLogMessageListRow.h
--- a/Engine/Source/Developer/MessageLog/Private/UserInterface/SMessageLogMessageListRow.h
+++ b/Engine/Source/Developer/MessageLog/Private/UserInterface/SMessageLogMessageListRow.h
@@ -154,6 +154,7 @@ public:
 						+ SVerticalBox::Slot()
 							.AutoHeight()
 							.HAlign(HAlign_Right)
+							.Padding(0.0f, 2.0f, 0.0f, 12.0f)
 							[
 								!HasTokens
 									? SNullWidget::NullWidget
@@ -274,7 +275,7 @@ protected:
 			InHorzBox->AddSlot()
 				.AutoWidth()
 				.VAlign(VAlign_Center)
-				.Padding(16.f, 4.0f, 0.0f, 12.0f)
+				.Padding(16.f, 0.0f, 0.0f, 0.0f)
 				[
 					SNew(SHorizontalBox)
 
