clang fixes

diff --git a/Engine/Source/Developer/MessageLog/Private/MessageLogPrivatePCH.h b/Engine/Source/Developer/MessageLog/Private/MessageLogPrivatePCH.h
--- a/Engine/Source/Developer/MessageLog/Private/MessageLogPrivatePCH.h
+++ b/Engine/Source/Developer/MessageLog/Private/MessageLogPrivatePCH.h
@@ -7,5 +7,6 @@
 #include "Slate.h"
 #include "Attribute.h"
 #include "EditorStyle.h"
+#include "MessageLogModule.h"
 
 #endif // __MessageLogPrivatePCH_h__
diff --git a/Engine/Source/Developer/MessageLog/Private/UserInterface/SMessageLogMessageListRow.h b/Engine/Source/Developer/MessageLog/Private/UserInterface/SMessageLogMessageListRow.h
--- a/Engine/Source/Developer/MessageLog/Private/UserInterface/SMessageLogMessageListRow.h
+++ b/Engine/Source/Developer/MessageLog/Private/UserInterface/SMessageLogMessageListRow.h
@@ -113,8 +113,8 @@ public:
 					[
 						(SeverityImageName == NAME_None)
 							? SNullWidget::NullWidget
-							: SNew(SImage)
-								.Image(FEditorStyle::GetBrush(SeverityImageName))
+							: static_cast<TSharedRef<SWidget>>(SNew(SImage)
+								.Image(FEditorStyle::GetBrush(SeverityImageName)))
 					]
 			]
 
@@ -148,7 +148,7 @@ public:
 							[
 								!HasTokens
 									? SNullWidget::NullWidget
-									: SNew(SSeparator)
+									: static_cast<TSharedRef<SWidget>>(SNew(SSeparator))
 							]
 
 						+ SVerticalBox::Slot()
@@ -158,7 +158,7 @@ public:
 							[
 								!HasTokens
 									? SNullWidget::NullWidget
-									: TokenBox
+									: static_cast<TSharedRef<SWidget>>(TokenBox)
 							]
 					]
 			];
