MessageLog: Added message tokens for tutorials

diff --git a/Engine/Source/Developer/MessageLog/MessageLog.Build.cs b/Engine/Source/Developer/MessageLog/MessageLog.Build.cs
--- a/Engine/Source/Developer/MessageLog/MessageLog.Build.cs
+++ b/Engine/Source/Developer/MessageLog/MessageLog.Build.cs
@@ -28,9 +28,10 @@ public class MessageLog : ModuleRules
 
 		if (UEBuildConfiguration.bBuildEditor)
 		{
-			PrivateDependencyModuleNames.AddRange(
+			PrivateIncludePathModuleNames.AddRange(
 				new string[] {
 					"Documentation",
+					"IntroTutorials",
 				}
 			);		
 		}
diff --git a/Engine/Source/Developer/MessageLog/Private/UserInterface/SMessageLogMessageListRow.h b/Engine/Source/Developer/MessageLog/Private/UserInterface/SMessageLogMessageListRow.h
--- a/Engine/Source/Developer/MessageLog/Private/UserInterface/SMessageLogMessageListRow.h
+++ b/Engine/Source/Developer/MessageLog/Private/UserInterface/SMessageLogMessageListRow.h
@@ -6,6 +6,7 @@
 
 #if WITH_EDITOR
 	#include "IDocumentation.h"
+	#include "IIntroTutorials.h"
 #endif
 
 #define LOCTEXT_NAMESPACE "SMessageLogMessageListRow"
@@ -249,9 +250,9 @@ protected:
 			}
 			break;
 
+#if WITH_EDITOR
 		case EMessageToken::Documentation:
 			{
-#if WITH_EDITOR
 				const TSharedRef<FDocumentationToken> DocumentationToken = StaticCastSharedRef<FDocumentationToken>(InMessageToken);
 
 				IconBrushName = FName("MessageLog.Docs");
@@ -265,9 +266,22 @@ protected:
 					)
 					.TextStyle(FEditorStyle::Get(), "MessageLog")
 					.OnNavigate(this, &SMessageLogMessageListRow::HandleDocsHyperlinkNavigate, DocumentationToken->GetDocumentationLink());
-#endif
 			}
 			break;
+
+		case EMessageToken::Tutorial:
+			{
+				const TSharedRef<FTutorialToken> TutorialToken = StaticCastSharedRef<FTutorialToken>(InMessageToken);
+
+				IconBrushName = FName("MessageLog.Tutorial");
+				Content = SNew(SHyperlink)
+					.Text(LOCTEXT("TutorialLabel", "Tutorial"))
+					.ToolTipText(LOCTEXT("TutorialTokenToolTip", "Click to open tutorial"))
+					.TextStyle(FEditorStyle::Get(), "MessageLog")
+					.OnNavigate(this, &SMessageLogMessageListRow::HandleTutorialHyperlinkNavigate, TutorialToken->GetTutorialAssetName());
+			}
+			break;
+#endif
 		}
 
 		if (Content.IsValid() && (IconBrushName != NAME_None))
@@ -306,13 +320,6 @@ private:
 		ActionToken->ExecuteAction();
 	}
 
-#if WITH_EDITOR
-	void HandleDocsHyperlinkNavigate( FString DocumentationLink )
-	{
-		IDocumentation::Get()->Open(DocumentationLink);
-	}
-#endif
-
 	void HandleHyperlinkNavigate( TSharedRef<IMessageToken> InMessageToken )
 	{
 		InMessageToken->GetOnMessageTokenActivated().ExecuteIfBound(InMessageToken);
@@ -326,6 +333,18 @@ private:
 		return FReply::Handled();
 	}
 
+#if WITH_EDITOR
+	void HandleDocsHyperlinkNavigate( FString DocumentationLink )
+	{
+		IDocumentation::Get()->Open(DocumentationLink);
+	}
+
+	void HandleTutorialHyperlinkNavigate( FString TutorialAssetName )
+	{
+		IIntroTutorials::Get().LaunchTutorial(TutorialAssetName);
+	}
+#endif
+
 protected:
 
 	/** The message used to create this widget. */
diff --git a/Engine/Source/Editor/EditorStyle/Private/SlateEditorStyle.cpp b/Engine/Source/Editor/EditorStyle/Private/SlateEditorStyle.cpp
--- a/Engine/Source/Editor/EditorStyle/Private/SlateEditorStyle.cpp
+++ b/Engine/Source/Editor/EditorStyle/Private/SlateEditorStyle.cpp
@@ -840,6 +840,7 @@ void FSlateEditorStyle::FStyle::SetupGeneralStyles()
 		Set( "MessageLog.Error", new IMAGE_BRUSH( "Old/Kismet2/Log_Error", Icon16x16 ) );
 		Set( "MessageLog.Warning", new IMAGE_BRUSH( "Old/Kismet2/Log_Warning", Icon16x16 ) );
 		Set( "MessageLog.Note", new IMAGE_BRUSH( "Old/Kismet2/Log_Note", Icon16x16 ) );
+		Set( "MessageLog.Tutorial", new IMAGE_BRUSH( "Icons/icon_Blueprint_Enum_16x", Icon16x16 ) );
 		Set( "MessageLog.Url", new IMAGE_BRUSH( "Icons/icon_world_16x", Icon16x16 ) );
 
 		Set( "MessageLog.TabIcon", new IMAGE_BRUSH( "Icons/icon_tab_MessageLog_16x", Icon16x16 ) );
diff --git a/Engine/Source/Editor/IntroTutorials/Private/IntroTutorials.cpp b/Engine/Source/Editor/IntroTutorials/Private/IntroTutorials.cpp
--- a/Engine/Source/Editor/IntroTutorials/Private/IntroTutorials.cpp
+++ b/Engine/Source/Editor/IntroTutorials/Private/IntroTutorials.cpp
@@ -328,9 +328,9 @@ void FIntroTutorials::OnNewProjectDialogOpened()
 void FIntroTutorials::HandleCompilerNotFound()
 {
 #if PLATFORM_WINDOWS
-	LaunchTutorial( TEXT( "Engine/Tutorial/Installation/InstallingVisualStudioTutorial.InstallingVisualStudioTutorial" ) );
+	LaunchTutorialByName( TEXT( "Engine/Tutorial/Installation/InstallingVisualStudioTutorial.InstallingVisualStudioTutorial" ) );
 #elif PLATFORM_MAC
-	LaunchTutorial( TEXT( "Engine/Tutorial/Installation/InstallingXCodeTutorial.InstallingXCodeTutorial" ) );
+	LaunchTutorialByName( TEXT( "Engine/Tutorial/Installation/InstallingXCodeTutorial.InstallingXCodeTutorial" ) );
 #else
 	STUBBED("FIntroTutorials::HandleCompilerNotFound");
 #endif
@@ -338,7 +338,7 @@ void FIntroTutorials::HandleCompilerNotFound()
 
 void FIntroTutorials::HandleSDKNotInstalled(const FString& PlatformName, const FString& InTutorialAsset)
 {
-	LaunchTutorial( InTutorialAsset );
+	LaunchTutorialByName( InTutorialAsset );
 }
 
 EVisibility FIntroTutorials::GetHomeButtonVisibility() const
@@ -364,7 +364,12 @@ void FIntroTutorials::SummonTutorialBrowser(TWeakPtr<SWindow> InWindow, const FS
 	}
 }
 
-void FIntroTutorials::LaunchTutorial(const FString& InAssetPath, bool bInRestart, TWeakPtr<SWindow> InNavigationWindow, FSimpleDelegate OnTutorialClosed, FSimpleDelegate OnTutorialExited)
+void FIntroTutorials::LaunchTutorial(const FString& TutorialAssetName)
+{
+	LaunchTutorialByName(TutorialAssetName);
+}
+
+void FIntroTutorials::LaunchTutorialByName(const FString& InAssetPath, bool bInRestart, TWeakPtr<SWindow> InNavigationWindow, FSimpleDelegate OnTutorialClosed, FSimpleDelegate OnTutorialExited)
 {
 	UBlueprint* Blueprint = LoadObject<UBlueprint>(nullptr, *InAssetPath);
 	if (Blueprint && Blueprint->GeneratedClass)
diff --git a/Engine/Source/Editor/IntroTutorials/Private/IntroTutorials.h b/Engine/Source/Editor/IntroTutorials/Private/IntroTutorials.h
--- a/Engine/Source/Editor/IntroTutorials/Private/IntroTutorials.h
+++ b/Engine/Source/Editor/IntroTutorials/Private/IntroTutorials.h
@@ -78,11 +78,12 @@ private:
 	EVisibility GetHomeButtonVisibility() const;
 
 	/** Internal helper to launch a tutorial from a path */
-	void LaunchTutorial(const FString& InAssetPath, bool bRestart = true, TWeakPtr<SWindow> InNavigationWindow = nullptr, FSimpleDelegate OnTutorialClosed = FSimpleDelegate(), FSimpleDelegate OnTutorialExited = FSimpleDelegate());
+	void LaunchTutorialByName(const FString& InAssetPath, bool bRestart = true, TWeakPtr<SWindow> InNavigationWindow = nullptr, FSimpleDelegate OnTutorialClosed = FSimpleDelegate(), FSimpleDelegate OnTutorialExited = FSimpleDelegate());
 
 public:
 
 	// IIntroTutorials interface
+	virtual void LaunchTutorial(const FString& TutorialAssetName) override;
 	virtual void LaunchTutorial(UEditorTutorial* Tutorial, bool bRestart = true, TWeakPtr<SWindow> InNavigationWindow = nullptr, FSimpleDelegate OnTutorialClosed = FSimpleDelegate(), FSimpleDelegate OnTutorialExited = FSimpleDelegate()) override;
 	virtual void CloseAllTutorialContent() override;
 	virtual TSharedRef<SWidget> CreateTutorialsWidget(FName InContext, TWeakPtr<SWindow> InContextWindow = nullptr) const override;
diff --git a/Engine/Source/Editor/IntroTutorials/Public/IIntroTutorials.h b/Engine/Source/Editor/IntroTutorials/Public/IIntroTutorials.h
--- a/Engine/Source/Editor/IntroTutorials/Public/IIntroTutorials.h
+++ b/Engine/Source/Editor/IntroTutorials/Public/IIntroTutorials.h
@@ -36,6 +36,13 @@ public:
 		return FModuleManager::Get().IsModuleLoaded( "IntroTutorials" );
 	}
 
+	/**
+	 * Launch a tutorial with the specified asset.
+	 *
+	 * @param TutorialAssetName The name of the tutorial asset.
+	 */
+	virtual void LaunchTutorial(const FString& TutorialAssetName) = 0;
+
 	/**
 	 * Launch a tutorial immediately, bypassing the tutorial browser.
 	 * 
diff --git a/Engine/Source/Runtime/Core/Public/Logging/TokenizedMessage.h b/Engine/Source/Runtime/Core/Public/Logging/TokenizedMessage.h
--- a/Engine/Source/Runtime/Core/Public/Logging/TokenizedMessage.h
+++ b/Engine/Source/Runtime/Core/Public/Logging/TokenizedMessage.h
@@ -27,13 +27,14 @@ namespace EMessageToken
 	enum Type
 	{
 		Action,
-		Text,
+		AssetName,
+		Documentation,
 		Image,
-		Severity,
 		Object,
+		Severity,
+		Text,
+		Tutorial,
 		URL,
-		AssetName,
-		Documentation,
 	};
 }
 
@@ -498,4 +499,43 @@ private:
 
 	/** The action's description text. */
 	const FText ActionDescription;
+};
+
+
+class FTutorialToken
+	: public IMessageToken
+{
+public:
+
+	/** Factory method, tokens can only be constructed as shared refs */
+	CORE_API static TSharedRef<FTutorialToken> Create( const FString& TutorialAssetName )
+	{
+		return MakeShareable(new FTutorialToken(TutorialAssetName));
+	}
+
+public:
+
+	// IMessageToken interface
+
+	virtual EMessageToken::Type GetType() const override
+	{
+		return EMessageToken::Tutorial;
+	}
+
+	/** Get the tutorial asset name stored in this token. */
+	const FString& GetTutorialAssetName() const
+	{
+		return TutorialAssetName;
+	}
+
+protected:
+	/** Protected constructor */
+	FTutorialToken( const FString& InTutorialAssetName )
+		: TutorialAssetName(InTutorialAssetName)
+	{ }
+
+private:
+
+	/** The name of the tutorial asset. */
+	FString TutorialAssetName;
 };
\ No newline at end of file
