Message log now opens the correct page when popped up

Handled log listing change broadcast message & selected the correct log.
Also corrected log list to use localized title.

TTP# 347358 - EDITOR: Message log (sometimes?) does not select the correct category

diff --git a/Engine/Source/Developer/MessageLog/Private/UserInterface/SMessageLog.cpp b/Engine/Source/Developer/MessageLog/Private/UserInterface/SMessageLog.cpp
--- a/Engine/Source/Developer/MessageLog/Private/UserInterface/SMessageLog.cpp
+++ b/Engine/Source/Developer/MessageLog/Private/UserInterface/SMessageLog.cpp
@@ -16,7 +16,7 @@
 
 SMessageLog::~SMessageLog()
 {
-	//ViewModel->OnSelectionChanged().RemoveAll(this);
+	ViewModel->OnSelectionChanged().RemoveAll(this);
 }
 
 
@@ -90,10 +90,19 @@ void SMessageLog::Construct( const FArguments& InArgs, const TSharedRef<FMessage
 		CategoriesListView->SetSelection(DefaultViewModel, ESelectInfo::Direct);
 	}
 
-//	ViewModel->OnSelectionChanged().AddSP(this, &SMessageLog::OnSelectionUpdated);
+	ViewModel->OnSelectionChanged().AddSP(this, &SMessageLog::HandleSelectionUpdated);
 }
 END_SLATE_FUNCTION_BUILD_OPTIMIZATION
 
+void SMessageLog::HandleSelectionUpdated()
+{
+	TSharedPtr<FMessageLogListingViewModel> LogListingViewModel = ViewModel->GetCurrentListingViewModel();
+
+	if ( LogListingViewModel.IsValid() )
+	{
+		CategoriesListView->SetSelection(LogListingViewModel, ESelectInfo::Direct);
+	}
+}
 
 /* SWidget overrides
  *****************************************************************************/
diff --git a/Engine/Source/Developer/MessageLog/Private/UserInterface/SMessageLog.h b/Engine/Source/Developer/MessageLog/Private/UserInterface/SMessageLog.h
--- a/Engine/Source/Developer/MessageLog/Private/UserInterface/SMessageLog.h
+++ b/Engine/Source/Developer/MessageLog/Private/UserInterface/SMessageLog.h
@@ -42,6 +42,9 @@ private:
 	/** Callback for changing the selected item in the categories list. */
 	void HandleCategoriesListSelectionChanged( IMessageLogListingPtr Selection, ESelectInfo::Type SelectInfo );
 
+	/** Handles the broadcast message sent the current log listing is changed */
+	void HandleSelectionUpdated();
+
 private:	
 
 	/** Holds the log categories list view widget. */
diff --git a/Engine/Source/Developer/MessageLog/Private/UserInterface/SMessageLogCategoryListRow.h b/Engine/Source/Developer/MessageLog/Private/UserInterface/SMessageLogCategoryListRow.h
--- a/Engine/Source/Developer/MessageLog/Private/UserInterface/SMessageLogCategoryListRow.h
+++ b/Engine/Source/Developer/MessageLog/Private/UserInterface/SMessageLogCategoryListRow.h
@@ -63,15 +63,14 @@ private:
 	// Callback for getting the text in the 'Entries' column.
 	FText HandleNameColumnText( ) const
 	{
-		const FText CategoryName = FText::FromString(LogListing->GetName().ToString());
 		int32 NumMessages = LogListing->GetFilteredMessages().Num();
 
 		if (NumMessages == 0)
 		{
-			return FText::Format(LOCTEXT("ColumnNameFormat", "{0}"), CategoryName);
+			return LogListing->GetLabel();
 		}
 
-		return FText::Format(LOCTEXT("ColumnNameCountFormat", "{0} ({1})"), CategoryName, FText::AsNumber(NumMessages));		
+		return FText::Format(LOCTEXT("ColumnNameCountFormat", "{0} ({1})"), LogListing->GetLabel(), FText::AsNumber(NumMessages));		
 	}
 
 	// Callback for getting the text color.
