MessageLog: Redesigned message log layout for a more condensed view; fixed some message texts

diff --git a/Engine/Source/Developer/MessageLog/Private/UserInterface/SMessageLogMessageListRow.h b/Engine/Source/Developer/MessageLog/Private/UserInterface/SMessageLogMessageListRow.h
--- a/Engine/Source/Developer/MessageLog/Private/UserInterface/SMessageLogMessageListRow.h
+++ b/Engine/Source/Developer/MessageLog/Private/UserInterface/SMessageLogMessageListRow.h
@@ -50,6 +50,8 @@ public:
 		);
 	}
 
+public:
+
 	/** @return Widget for this log listing item*/
 	virtual TSharedRef<SWidget> GenerateWidget()
 	{
@@ -57,48 +59,40 @@ public:
 		const TArray<TSharedRef<IMessageToken>>& MessageTokens = Message->GetMessageTokens();
 
 		// Create the horizontal box and add the icon
-		TSharedRef<SVerticalBox> DetailBox = SNew(SVerticalBox);
-		TSharedRef<SHorizontalBox> TokenBox = SNew(SHorizontalBox);
+		TSharedRef<SHorizontalBox> MessageBox = SNew(SHorizontalBox);
+		TSharedRef<SHorizontalBox> LinkBox = SNew(SHorizontalBox);
 		FName SeverityImageName = NAME_None;
-		bool HasTokens = false;
-		FText TitleText;
+		bool HasLinks = false;
 
 		// Iterate over parts of the message and create widgets for them
 		for (auto TokenIt = MessageTokens.CreateConstIterator(); TokenIt; ++TokenIt)
 		{
 			const TSharedRef<IMessageToken>& Token = *TokenIt;
 
-			if (Token->GetType() == EMessageToken::Severity)
-			{
-				if (SeverityImageName == NAME_None)
-				{
-					const TSharedRef<FSeverityToken> SeverityToken = StaticCastSharedRef<FSeverityToken>(Token);
-					SeverityImageName = FTokenizedMessage::GetSeverityIconName(SeverityToken->GetSeverity());
-				}
-			}
-			else if (Token->GetType() == EMessageToken::Text)
+			switch (Token->GetType())
 			{
-				if (TitleText.IsEmpty())
+			case EMessageToken::Severity:
 				{
-					TitleText = Token->ToText();
+					if (SeverityImageName == NAME_None)
+					{
+						const TSharedRef<FSeverityToken> SeverityToken = StaticCastSharedRef<FSeverityToken>(Token);
+						SeverityImageName = FTokenizedMessage::GetSeverityIconName(SeverityToken->GetSeverity());
+					}
 				}
-				else
+				break;
+
+			case EMessageToken::Documentation:
+			case EMessageToken::Tutorial:
 				{
-					DetailBox->AddSlot()
-						.AutoHeight()
-						.Padding(0.0f, 0.0f, 0.0f, 6.0f)
-						[
-							SNew(STextBlock)
-								.ColorAndOpacity(FSlateColor::UseSubduedForeground())
-								.Text(Token->ToText())
-						];
+					CreateMessage(LinkBox, Token, 10.0f);
+					HasLinks = true;
 				}
+				break;
+
+			default:
+				CreateMessage(MessageBox, Token, 2.0f);
+				break;
 			}
-			else
-			{
-				CreateMessage(TokenBox, Token);
-				HasTokens = true;
-			}		
 		}
 
 		return SNew(SHorizontalBox)
@@ -106,11 +100,11 @@ public:
 
 		+ SHorizontalBox::Slot()
 			.AutoWidth()
+			.HAlign(HAlign_Center)
+			.VAlign(VAlign_Center)
 			[
 				SNew(SBox)
 					.Padding(2.0f)
-					.HAlign(HAlign_Center)
-					.VAlign(VAlign_Top)
 					[
 						(SeverityImageName == NAME_None)
 							? SNullWidget::NullWidget
@@ -121,47 +115,25 @@ public:
 
 		+ SHorizontalBox::Slot()
 			.FillWidth(1.0f)
-			.Padding(2.0f, 0.0f, 0.0f, 0.0f)
+			.VAlign(VAlign_Center)
 			[
-				SNew(SBox)
-					.Padding(2.0f)
-					[
-						SNew(SVerticalBox)
-
-						+ SVerticalBox::Slot()
-							.AutoHeight()
-							.VAlign(VAlign_Center)
-							[
-								SNew(STextBlock)
-									.Font(FSlateFontInfo(FPaths::EngineContentDir() / TEXT("Slate/Fonts/Roboto-Bold.ttf"), 9))
-									.Text(TitleText)
-							]
-
-						+ SVerticalBox::Slot()
-							.AutoHeight()
-							.Padding(0.0f, 4.0f, 0.0f, 0.0f)
-							[
-								DetailBox
-							]
+				MessageBox
+			]
 
-						+ SVerticalBox::Slot()
-							.AutoHeight()
-							[
-								!HasTokens
-									? SNullWidget::NullWidget
-									: static_cast<TSharedRef<SWidget>>(SNew(SSeparator))
-							]
-
-						+ SVerticalBox::Slot()
-							.AutoHeight()
-							.HAlign(HAlign_Right)
-							.Padding(0.0f, 2.0f, 0.0f, 12.0f)
-							[
-								!HasTokens
-									? SNullWidget::NullWidget
-									: static_cast<TSharedRef<SWidget>>(TokenBox)
-							]
-					]
+		+ SHorizontalBox::Slot()
+			.AutoWidth()
+			.HAlign(HAlign_Right)
+			.VAlign(VAlign_Center)
+			.Padding(1.0f)
+			[
+				!HasLinks
+					? SNullWidget::NullWidget
+					: static_cast<TSharedRef<SWidget>>(SNew(SBorder)
+						.BorderImage(FEditorStyle::GetBrush("ToolPanel.GroupBorder"))
+						.Padding(FMargin(0.0f, 1.0f, 10.0f, 1.0f))
+						[
+							LinkBox
+						])
 			];
 	}
 
@@ -176,13 +148,13 @@ protected:
 			.OnNavigate(this, &SMessageLogMessageListRow::HandleHyperlinkNavigate, InMessageToken);
 	}
 
-	void CreateMessage( const TSharedRef<SHorizontalBox>& InHorzBox, const TSharedRef<IMessageToken>& InMessageToken )
+	void CreateMessage( const TSharedRef<SHorizontalBox>& InHorzBox, const TSharedRef<IMessageToken>& InMessageToken, float Padding )
 	{
 		TSharedPtr<SWidget> Content;
 		FName IconBrushName;
 
 		switch(InMessageToken->GetType())
-		{/*
+		{
 		case EMessageToken::Image:
 			{
 				const TSharedRef<FImageToken> ImageToken = StaticCastSharedRef<FImageToken>(InMessageToken);
@@ -207,7 +179,7 @@ protected:
 				}
 			}
 			break;
-*/
+
 		case EMessageToken::Object:
 			{
 				const TSharedRef<FUObjectToken> UObjectToken = StaticCastSharedRef<FUObjectToken>(InMessageToken);
@@ -269,6 +241,14 @@ protected:
 			}
 			break;
 
+		case EMessageToken::Text:
+			{
+				Content = SNew(STextBlock)
+					.ColorAndOpacity(FSlateColor::UseSubduedForeground())
+					.Text(InMessageToken->ToText());
+			}
+			break;
+
 		case EMessageToken::Tutorial:
 			{
 				const TSharedRef<FTutorialToken> TutorialToken = StaticCastSharedRef<FTutorialToken>(InMessageToken);
@@ -284,12 +264,12 @@ protected:
 #endif
 		}
 
-		if (Content.IsValid() && (IconBrushName != NAME_None))
+		if (Content.IsValid())
 		{
 			InHorzBox->AddSlot()
 				.AutoWidth()
 				.VAlign(VAlign_Center)
-				.Padding(16.f, 0.0f, 0.0f, 0.0f)
+				.Padding(Padding, 0.0f, 0.0f, 0.0f)
 				[
 					SNew(SHorizontalBox)
 
@@ -297,9 +277,11 @@ protected:
 						.AutoWidth()
 						.VAlign(VAlign_Center)
 						[
-							SNew(SImage)
-								.ColorAndOpacity(FSlateColor::UseForeground())
-								.Image(FEditorStyle::GetBrush(IconBrushName))
+							(IconBrushName == NAME_None)
+								? SNullWidget::NullWidget
+								: static_cast<TSharedRef<SWidget>>(SNew(SImage)
+									.ColorAndOpacity(FSlateColor::UseForeground())
+									.Image(FEditorStyle::GetBrush(IconBrushName)))
 						]
 
 					+ SHorizontalBox::Slot()
diff --git a/Engine/Source/Editor/UnrealEd/Private/EditorServer.cpp b/Engine/Source/Editor/UnrealEd/Private/EditorServer.cpp
--- a/Engine/Source/Editor/UnrealEd/Private/EditorServer.cpp
+++ b/Engine/Source/Editor/UnrealEd/Private/EditorServer.cpp
@@ -3463,8 +3463,7 @@ bool UEditorEngine::Map_Check( UWorld* InWorld, const TCHAR* Str, FOutputDevice&
 		Arguments.Add(TEXT("Warnings"), WarningCount - ErrorCount);
 		Arguments.Add(TEXT("Time"), CurrentTime);
 		FMessageLog("MapCheck").Info()
-			->AddToken(FTextToken::Create(LOCTEXT( "MapCheck_Complete", "Map check completed" )))
-			->AddToken(FTextToken::Create(FText::Format( LOCTEXT( "MapCheck_CompleteDetail", "There are {Errors} Error(s) and {Warnings} Warning(s), and it took {Time}ms to complete." ), Arguments ) ));
+			->AddToken(FTextToken::Create(FText::Format( LOCTEXT( "MapCheck_Complete", "Map check complete: {Errors} Error(s), {Warnings} Warning(s), took {Time}ms to complete." ), Arguments ) ));
 	}
 
 	GWarn->EndSlowTask();
diff --git a/Engine/Source/Runtime/Engine/Private/AmbientSound.cpp b/Engine/Source/Runtime/Engine/Private/AmbientSound.cpp
--- a/Engine/Source/Runtime/Engine/Private/AmbientSound.cpp
+++ b/Engine/Source/Runtime/Engine/Private/AmbientSound.cpp
@@ -52,20 +52,16 @@ void AAmbientSound::CheckForErrors( void )
 
 	if( !AudioComponent.IsValid() )
 	{
-		FFormatNamedArguments Arguments;
-		Arguments.Add(TEXT("ActorName"), FText::FromString(GetName()));
 		FMessageLog("MapCheck").Warning()
 			->AddToken(FUObjectToken::Create(this))
-			->AddToken(FTextToken::Create(FText::Format( LOCTEXT( "MapCheck_Message_AudioComponentNull", "{ActorName} : Ambient sound actor has NULL AudioComponent property - please delete" ), Arguments ) ))
+			->AddToken(FTextToken::Create(LOCTEXT( "MapCheck_Message_AudioComponentNull", "Ambient sound actor has NULL AudioComponent property - please delete")))
 			->AddToken(FMapErrorToken::Create(FMapErrors::AudioComponentNull));
 	}
 	else if( AudioComponent->Sound == NULL )
 	{
-		FFormatNamedArguments Arguments;
-		Arguments.Add(TEXT("ActorName"), FText::FromString(GetName()));
 		FMessageLog("MapCheck").Warning()
 			->AddToken(FUObjectToken::Create(this))
-			->AddToken(FTextToken::Create(FText::Format( LOCTEXT( "MapCheck_Message_SoundCueNull", "{ActorName} : Ambient sound actor has NULL Sound Cue property" ), Arguments ) ))
+			->AddToken(FTextToken::Create(LOCTEXT( "MapCheck_Message_SoundCueNull", "Ambient sound actor has NULL Sound Cue property")))
 			->AddToken(FMapErrorToken::Create(FMapErrors::SoundCueNull));
 	}
 }
diff --git a/Engine/Source/Runtime/Engine/Private/ErrorChecking.cpp b/Engine/Source/Runtime/Engine/Private/ErrorChecking.cpp
--- a/Engine/Source/Runtime/Engine/Private/ErrorChecking.cpp
+++ b/Engine/Source/Runtime/Engine/Private/ErrorChecking.cpp
@@ -188,11 +188,9 @@ void AVolume::CheckForErrors()
 
 	if (GetRootComponent() == NULL)
 	{
-		FFormatNamedArguments Arguments;
-		Arguments.Add(TEXT("ActorName"), FText::FromString(GetName()));
 		FMessageLog("MapCheck").Warning()
 			->AddToken(FUObjectToken::Create(this))
-			->AddToken(FTextToken::Create(FText::Format(LOCTEXT( "MapCheck_Message_VolumeActorCollisionComponentNULL", "{ActorName} : Volume actor has NULL collision component - please delete" ), Arguments) ))
+			->AddToken(FTextToken::Create(LOCTEXT( "MapCheck_Message_VolumeActorCollisionComponentNULL", "Volume actor has NULL collision component - please delete")))
 			->AddToken(FMapErrorToken::Create(FMapErrors::VolumeActorCollisionComponentNULL));
 	}
 	else
@@ -203,7 +201,7 @@ void AVolume::CheckForErrors()
 			Arguments.Add(TEXT("ActorName"), FText::FromString(GetName()));
 			FMessageLog("MapCheck").Warning()
 				->AddToken(FUObjectToken::Create(this))
-				->AddToken(FTextToken::Create(FText::Format(LOCTEXT( "MapCheck_Message_VolumeActorZeroRadius", "{ActorName} : Volume actor has a collision component with 0 radius - please delete" ), Arguments) ));
+				->AddToken(FTextToken::Create(LOCTEXT( "MapCheck_Message_VolumeActorZeroRadius", "Volume actor has a collision component with 0 radius - please delete")));
 		}
 	}
 }
@@ -218,7 +216,7 @@ void ALight::CheckForErrors()
 		Arguments.Add(TEXT("ActorName"), FText::FromString(GetName()));
 		FMessageLog("MapCheck").Warning()
 			->AddToken(FUObjectToken::Create(this))
-			->AddToken(FTextToken::Create(FText::Format(LOCTEXT( "MapCheck_Message_LightComponentNull", "{ActorName} : Light actor has NULL LightComponent property - please delete!" ), Arguments) ))
+			->AddToken(FTextToken::Create(LOCTEXT("MapCheck_Message_LightComponentNull", "Light actor has NULL LightComponent property - please delete!")))
 			->AddToken(FMapErrorToken::Create(FMapErrors::LightComponentNull));
 	}
 }
diff --git a/Engine/Source/Runtime/Engine/Private/Particles/ParticleComponents.cpp b/Engine/Source/Runtime/Engine/Private/Particles/ParticleComponents.cpp
--- a/Engine/Source/Runtime/Engine/Private/Particles/ParticleComponents.cpp
+++ b/Engine/Source/Runtime/Engine/Private/Particles/ParticleComponents.cpp
@@ -277,11 +277,9 @@ void AEmitter::CheckForErrors()
 	{
 		if ( !ParticleSystemComponent.IsValid() )
 		{
-			FFormatNamedArguments Arguments;
-			Arguments.Add(TEXT("EmitterName"), FText::FromString(GetName()));
 			FMessageLog("MapCheck").Warning()
 				->AddToken(FUObjectToken::Create(this))
-				->AddToken(FTextToken::Create(FText::Format( LOCTEXT( "MapCheck_Message_ParticleSystemComponentNull", "{EmitterName} : Emitter actor has NULL ParticleSystemComponent property - please delete" ), Arguments ) ))
+				->AddToken(FTextToken::Create(LOCTEXT("MapCheck_Message_ParticleSystemComponentNull", "Emitter actor has NULL ParticleSystemComponent property - please delete")))
 				->AddToken(FMapErrorToken::Create(FMapErrors::ParticleSystemComponentNull));
 		}
 	}
diff --git a/Engine/Source/Runtime/Engine/Private/PrimitiveComponent.cpp b/Engine/Source/Runtime/Engine/Private/PrimitiveComponent.cpp
--- a/Engine/Source/Runtime/Engine/Private/PrimitiveComponent.cpp
+++ b/Engine/Source/Runtime/Engine/Private/PrimitiveComponent.cpp
@@ -592,8 +592,7 @@ void UPrimitiveComponent::CheckForErrors()
 	{
 		FMessageLog("MapCheck").PerformanceWarning()
 			->AddToken(FUObjectToken::Create(Owner))
-			->AddToken(FTextToken::Create(LOCTEXT( "MapCheck_Message_ShadowCasterUsingBoundsScale", "Actor casts dynamic shadows and has a BoundsScale greater than 1" )))
-			->AddToken(FTextToken::Create(LOCTEXT( "MapCheck_Message_ShadowCasterUsingBoundsScaleDetail", "This will have a large performance hit." )))
+			->AddToken(FTextToken::Create(LOCTEXT( "MapCheck_Message_ShadowCasterUsingBoundsScale", "Actor casts dynamic shadows and has a BoundsScale greater than 1! This will have a large performance hit" )))
 			->AddToken(FMapErrorToken::Create(FMapErrors::ShadowCasterUsingBoundsScale));
 	}
 
@@ -601,9 +600,7 @@ void UPrimitiveComponent::CheckForErrors()
 	{
 		FMessageLog("MapCheck").Error()
 			->AddToken(FUObjectToken::Create(Owner))
-			->AddToken(FTextToken::Create(LOCTEXT( "MapCheck_Message_InvalidLightmapSettings", "Component is a static type but has invalid lightmap settings" )))
-			->AddToken(FTextToken::Create(LOCTEXT( "MapCheck_Message_InvalidLightmapSettingsDetail", "The indirect lighting on this component will be black." )))
-			->AddToken(FTextToken::Create(LOCTEXT( "MapCheck_Message_InvalidLightmapSettingsDetail2", "Common causes: lightmap resolution of 0, LightmapCoordinateIndex out of bounds" )))
+			->AddToken(FTextToken::Create(LOCTEXT( "MapCheck_Message_InvalidLightmapSettings", "Component is a static type but has invalid lightmap settings!  Indirect lighting will be black.  Common causes are lightmap resolution of 0, LightmapCoordinateIndex out of bounds." )))
 			->AddToken(FMapErrorToken::Create(FMapErrors::StaticComponentHasInvalidLightmapSettings));
 	}
 }
diff --git a/Engine/Source/Runtime/Engine/Private/SkeletalMesh.cpp b/Engine/Source/Runtime/Engine/Private/SkeletalMesh.cpp
--- a/Engine/Source/Runtime/Engine/Private/SkeletalMesh.cpp
+++ b/Engine/Source/Runtime/Engine/Private/SkeletalMesh.cpp
@@ -3593,11 +3593,9 @@ void ASkeletalMeshActor::CheckForErrors()
 
 		if (SkeletalMeshComponent->SkeletalMesh == NULL)
 		{
-			FFormatNamedArguments Arguments;
-			Arguments.Add(TEXT("ActorName"), FText::FromString(GetName()));
 			FMessageLog("MapCheck").Warning()
 				->AddToken(FUObjectToken::Create(this))
-				->AddToken(FTextToken::Create(FText::Format( LOCTEXT( "MapCheck_Message_SkeletalMeshNull", "{ActorName} : Skeletal mesh actor has NULL SkeletalMesh property" ), Arguments ) ))
+				->AddToken(FTextToken::Create(LOCTEXT("MapCheck_Message_SkeletalMeshNull", "Skeletal mesh actor has NULL SkeletalMesh property")))
 				->AddToken(FMapErrorToken::Create(FMapErrors::SkeletalMeshNull));
 		}
 	}
@@ -3607,7 +3605,7 @@ void ASkeletalMeshActor::CheckForErrors()
 		Arguments.Add(TEXT("ActorName"), FText::FromString(GetName()));
 		FMessageLog("MapCheck").Warning()
 			->AddToken(FUObjectToken::Create(this))
-			->AddToken(FTextToken::Create(FText::Format( LOCTEXT( "MapCheck_Message_SkeletalMeshComponent", "{ActorName} : Skeletal mesh actor has NULL SkeletalMeshComponent property" ), Arguments ) ))
+			->AddToken(FTextToken::Create(LOCTEXT("MapCheck_Message_SkeletalMeshComponent", "Skeletal mesh actor has NULL SkeletalMeshComponent property")))
 			->AddToken(FMapErrorToken::Create(FMapErrors::SkeletalMeshComponent));
 	}
 }
diff --git a/Engine/Source/Runtime/Engine/Private/StaticMeshActor.cpp b/Engine/Source/Runtime/Engine/Private/StaticMeshActor.cpp
--- a/Engine/Source/Runtime/Engine/Private/StaticMeshActor.cpp
+++ b/Engine/Source/Runtime/Engine/Private/StaticMeshActor.cpp
@@ -128,15 +128,14 @@ void AStaticMeshActor::CheckForErrors()
 	{
 		MapCheck.Warning()
 			->AddToken(FUObjectToken::Create(this))
-			->AddToken(FTextToken::Create(LOCTEXT( "MapCheck_Message_StaticMeshComponent", "Static mesh actor has NULL StaticMeshComponent property" )))
-			->AddToken(FTextToken::Create(LOCTEXT( "MapCheck_Message_StaticMeshComponentDetail", "Please delete the component or assign a valid one." )))
+			->AddToken(FTextToken::Create(LOCTEXT( "MapCheck_Message_StaticMeshComponent", "Static mesh actor has NULL StaticMeshComponent property - please delete" ) ))
 			->AddToken(FMapErrorToken::Create(FMapErrors::StaticMeshComponent));
 	}
 	else if( StaticMeshComponent->StaticMesh == NULL )
 	{
 		MapCheck.Warning()
 			->AddToken(FUObjectToken::Create(this))
-			->AddToken(FTextToken::Create(LOCTEXT( "MapCheck_Message_StaticMeshNull", "Static mesh actor has NULL StaticMesh property" )))
+			->AddToken(FTextToken::Create(LOCTEXT( "MapCheck_Message_StaticMeshNull", "Static mesh actor has NULL StaticMesh property" ) ))
 			->AddToken(FMapErrorToken::Create(FMapErrors::StaticMeshNull));
 	}
 	else
@@ -161,8 +160,7 @@ void AStaticMeshActor::CheckForErrors()
 					Arguments.Add(TEXT("ActorName1"), FText::FromString(A->GetName()));
 					MapCheck.Warning()
 						->AddToken(FUObjectToken::Create(this))
-						->AddToken(FTextToken::Create(LOCTEXT( "MapCheck_Message_SameLocation", "Overlapping actors" )))
-						->AddToken(FTextToken::Create(FText::Format( LOCTEXT( "MapCheck_Message_SameLocationDetail", "{ActorName0} is in the same location as {ActorName1}." ), Arguments ) ))
+						->AddToken(FTextToken::Create(FText::Format( LOCTEXT( "MapCheck_Message_SameLocation", "{ActorName0} is in the same location as {ActorName1}" ), Arguments ) ))
 						->AddToken(FMapErrorToken::Create(FMapErrors::SameLocation));
 				}
 			}
diff --git a/Engine/Source/Runtime/Engine/Private/StaticMeshComponent.cpp b/Engine/Source/Runtime/Engine/Private/StaticMeshComponent.cpp
--- a/Engine/Source/Runtime/Engine/Private/StaticMeshComponent.cpp
+++ b/Engine/Source/Runtime/Engine/Private/StaticMeshComponent.cpp
@@ -352,7 +352,7 @@ void UStaticMeshComponent::CheckForErrors()
 	{
 		FMessageLog("MapCheck").Warning()
 			->AddToken(FUObjectToken::Create(Owner))
-			->AddToken(FTextToken::Create(LOCTEXT( "MapCheck_Message_StaticMeshNull", "Static mesh actor has NULL StaticMesh property" )))
+			->AddToken(FTextToken::Create(LOCTEXT("MapCheck_Message_StaticMeshNull", "Static mesh actor has NULL StaticMesh property")))
 			->AddToken(FMapErrorToken::Create(FMapErrors::StaticMeshNull));
 	}
 
@@ -380,7 +380,7 @@ void UStaticMeshComponent::CheckForErrors()
 	{
 		FMessageLog("MapCheck").Warning()
 			->AddToken(FUObjectToken::Create(this))
-			->AddToken(FTextToken::Create(LOCTEXT( "MapCheck_Message_SimulatePhyNoSimpleCollision", "Using bSimulatePhysics but StaticMesh has no simple collision.")));
+			->AddToken(FTextToken::Create(FText::Format(LOCTEXT( "MapCheck_Message_SimulatePhyNoSimpleCollision", "{0} : Using bSimulatePhysics but StaticMesh has not simple collision."), FText::FromString(GetName()) ) ));
 	}
 
 	if( Mobility == EComponentMobility::Movable &&
