Allowed FTextTokens to create hyperlinks in the message log if they have had an OnMessageTokenActivated delegate supplied.

diff --git a/Engine/Source/Developer/MessageLog/Private/UserInterface/SMessageLogMessageListRow.cpp b/Engine/Source/Developer/MessageLog/Private/UserInterface/SMessageLogMessageListRow.cpp
--- a/Engine/Source/Developer/MessageLog/Private/UserInterface/SMessageLogMessageListRow.cpp
+++ b/Engine/Source/Developer/MessageLog/Private/UserInterface/SMessageLogMessageListRow.cpp
@@ -218,9 +218,16 @@ void SMessageLogMessageListRow::CreateMessage(const TSharedRef<SHorizontalBox>&
 
 	case EMessageToken::Text:
 	{
-		Content = SNew(STextBlock)
-			.ColorAndOpacity(FSlateColor::UseSubduedForeground())
-			.Text(InMessageToken->ToText());
+		if (InMessageToken->GetOnMessageTokenActivated().IsBound())
+		{
+			Content = CreateHyperlink(InMessageToken, InMessageToken->ToText());
+		}
+		else
+		{
+			Content = SNew(STextBlock)
+				.ColorAndOpacity(FSlateColor::UseSubduedForeground())
+				.Text(InMessageToken->ToText());
+		}
 	}
 		break;
 
