Added a horizontal scrollbar to Message Log list view so that very long lines can be seen in their entirety.

diff --git a/Engine/Source/Developer/MessageLog/Private/UserInterface/SMessageLogListing.cpp b/Engine/Source/Developer/MessageLog/Private/UserInterface/SMessageLogListing.cpp
--- a/Engine/Source/Developer/MessageLog/Private/UserInterface/SMessageLogListing.cpp
+++ b/Engine/Source/Developer/MessageLog/Private/UserInterface/SMessageLogListing.cpp
@@ -30,6 +30,8 @@ void SMessageLogListing::Construct( const FArguments& InArgs, const TSharedRef<
 	TSharedPtr<SHorizontalBox> HorizontalBox = NULL;
 	TSharedPtr<SVerticalBox> VerticalBox = NULL;
 
+	TSharedRef<SScrollBar> ScrollBar = SNew(SScrollBar);
+
 	ChildSlot
 	[
 		SNew(SBorder)
@@ -44,11 +46,31 @@ void SMessageLogListing::Construct( const FArguments& InArgs, const TSharedRef<
 						SNew(SBorder)
 							.BorderImage(FEditorStyle::GetBrush("MessageLog.ListBorder"))
 							[
-								SAssignNew(MessageListView, SListView< TSharedRef<FTokenizedMessage> >)
-									.ListItemsSource(&MessageLogListingViewModel->GetFilteredMessages())
-									.OnGenerateRow(this, &SMessageLogListing::MakeMessageLogListItemWidget)
-									.OnSelectionChanged(this, &SMessageLogListing::OnLineSelectionChanged)
-									.ItemHeight(24.0f)
+								SNew(SHorizontalBox)
+								+SHorizontalBox::Slot()
+								.FillWidth(1)
+								[
+									SNew(SScrollBox)
+									.Orientation(EOrientation::Orient_Horizontal)
+									+ SScrollBox::Slot()
+									[
+										SAssignNew(MessageListView, SListView< TSharedRef<FTokenizedMessage> >)
+										.ListItemsSource(&MessageLogListingViewModel->GetFilteredMessages())
+										.OnGenerateRow(this, &SMessageLogListing::MakeMessageLogListItemWidget)
+										.OnSelectionChanged(this, &SMessageLogListing::OnLineSelectionChanged)
+										.ExternalScrollbar(ScrollBar)
+										.ItemHeight(24.0f)
+									]
+								]
+								+SHorizontalBox::Slot()
+								.AutoWidth()
+								[
+									SNew(SBox)
+									.WidthOverride(FOptionalSize(16))
+									[
+										ScrollBar
+									]
+								]
 							]
 					]
 			]
