#UE4 Message log notifications that contain errors or worse are now persistent until dismissed interactively. Extended the duration of messages that are less severe than error by 6 seconds.

diff --git a/Engine/Source/Developer/MessageLog/Private/Presentation/MessageLogListingViewModel.cpp b/Engine/Source/Developer/MessageLog/Private/Presentation/MessageLogListingViewModel.cpp
--- a/Engine/Source/Developer/MessageLog/Private/Presentation/MessageLogListingViewModel.cpp
+++ b/Engine/Source/Developer/MessageLog/Private/Presentation/MessageLogListingViewModel.cpp
@@ -9,6 +9,8 @@
 
 #define LOCTEXT_NAMESPACE "MessageLog"
 
+int32 FMessageLogListingViewModel::NextNotificationId = 0;
+
 FMessageLogListingViewModel::~FMessageLogListingViewModel()
 {
 	check(MessageLogListingModel.IsValid());
@@ -88,6 +90,27 @@ const TArray< TSharedRef< class FMessageFilter> >& FMessageLogListingViewModel::
 	return MessageFilters;
 }
 
+void FMessageLogListingViewModel::DismissNotification(int32 NotificationId)
+{
+	int32 NotificationIndex = OpenNotifications.IndexOfByPredicate([NotificationId](const FOpenNotification& N) { return N.NotificationId == NotificationId; });
+	if (NotificationIndex != INDEX_NONE)
+	{
+		TSharedPtr<SNotificationItem> NotificationItem = OpenNotifications[NotificationIndex].NotificationItem.Pin();
+		if (NotificationItem.IsValid())
+		{
+			NotificationItem->ExpireAndFadeout();
+		}
+
+		OpenNotifications.RemoveAtSwap(NotificationIndex);
+	}
+}
+
+void FMessageLogListingViewModel::OpenMessageLogFromNotification(int32 NotificationId)
+{
+	DismissNotification(NotificationId);
+	OpenMessageLog();
+}
+
 void FMessageLogListingViewModel::OpenMessageLog( )
 {
 	FMessageLogModule& MessageLogModule = FModuleManager::GetModuleChecked<FMessageLogModule>("MessageLog");
@@ -288,13 +311,35 @@ void FMessageLogListingViewModel::NotifyIfAnyMessages( const FText& Message, EMe
 
 		FNotificationInfo ErrorNotification(NotificationMessage);
 		ErrorNotification.Image = FEditorStyle::GetBrush(FTokenizedMessage::GetSeverityIconName(HighestSeverityPresent(0)));
-		ErrorNotification.bFireAndForget = true;
-		ErrorNotification.Hyperlink = FSimpleDelegate::CreateSP(this, &FMessageLogListingViewModel::OpenMessageLog);
-		ErrorNotification.HyperlinkText = LOCTEXT("ShowMessageLogHyperlink", "Show Message Log");
-		ErrorNotification.ExpireDuration = 8.0f; // Need this message to last a little longer than normal since the user may want to "Show Log"
-		ErrorNotification.bUseThrobber = true;
 
-		FSlateNotificationManager::Get().AddNotification(ErrorNotification);
+		if (NumMessagesPresent(0, EMessageSeverity::Error) > 0)
+		{
+			int32 NotificationId = NextNotificationId++;
+
+			ErrorNotification.bFireAndForget = false;
+			ErrorNotification.bUseThrobber = false;
+			ErrorNotification.FadeOutDuration = 0.f;
+			ErrorNotification.ExpireDuration = 0.f;
+			ErrorNotification.ButtonDetails.Add(FNotificationButtonInfo(LOCTEXT("DismissMessageButton", "Dismiss"), FText(), FSimpleDelegate::CreateSP(this, &FMessageLogListingViewModel::DismissNotification, NotificationId)));
+			ErrorNotification.ButtonDetails.Add(FNotificationButtonInfo(LOCTEXT("ShowMessageLogButton", "Show Message Log"), FText(), FSimpleDelegate::CreateSP(this, &FMessageLogListingViewModel::OpenMessageLogFromNotification, NotificationId)));
+
+			TSharedPtr<SNotificationItem> NewNotificationItem = FSlateNotificationManager::Get().AddNotification(ErrorNotification);
+			if (NewNotificationItem.IsValid())
+			{
+				NewNotificationItem->SetCompletionState(SNotificationItem::CS_Pending);
+				OpenNotifications.Emplace(NotificationId, NewNotificationItem);
+			}
+		}
+		else
+		{
+			ErrorNotification.bFireAndForget = true;
+			ErrorNotification.Hyperlink = FSimpleDelegate::CreateSP(this, &FMessageLogListingViewModel::OpenMessageLog);
+			ErrorNotification.HyperlinkText = LOCTEXT("ShowMessageLogHyperlink", "Show Message Log");
+			ErrorNotification.ExpireDuration = 14.0f; // Need this message to last a little longer than normal since the user may want to "Show Log"
+			ErrorNotification.bUseThrobber = true;
+
+			FSlateNotificationManager::Get().AddNotification(ErrorNotification);
+		}
 	}
 }
 
diff --git a/Engine/Source/Developer/MessageLog/Private/Presentation/MessageLogListingViewModel.h b/Engine/Source/Developer/MessageLog/Private/Presentation/MessageLogListingViewModel.h
--- a/Engine/Source/Developer/MessageLog/Private/Presentation/MessageLogListingViewModel.h
+++ b/Engine/Source/Developer/MessageLog/Private/Presentation/MessageLogListingViewModel.h
@@ -168,6 +168,12 @@ private:
 	/** Rebuilds the list of filtered messages */
 	void RefreshFilteredMessages();
 
+	/** Dismisses a notification that was previously shown. */
+	void DismissNotification(int32 NotificationId);
+
+	/** Opens the message log from a notification that was previously shown. */
+	void OpenMessageLogFromNotification(int32 NotificationId);
+
 	/** Helper function for opening this message log from a notification */
 	void OpenMessageLog();
 
@@ -229,4 +235,15 @@ private:
 
 	/** Delegate to call when page selection state is changed */
 	FOnPageSelectionChangedEvent PageSelectionChangedEvent;
+
+	/** All open notifications */
+	struct FOpenNotification
+	{
+		int32 NotificationId;
+		TWeakPtr<class SNotificationItem> NotificationItem;
+		FOpenNotification(int32 InNotificationId, const TWeakPtr<class SNotificationItem>& InNotificationItem) : NotificationId(InNotificationId), NotificationItem(InNotificationItem) {}
+		FOpenNotification() : NotificationId(INDEX_NONE) {}
+	};
+	TArray<FOpenNotification> OpenNotifications;
+	static int32 NextNotificationId;
 };
