Set the stats message log to use a fixed size font, and removed duplicated tooltip text.

# Integrating CL 12106172 from Dev-Editor to Release 4.25

diff --git a/Engine/Source/Developer/MessageLog/Private/UserInterface/SMessageLogMessageListRow.cpp b/Engine/Source/Developer/MessageLog/Private/UserInterface/SMessageLogMessageListRow.cpp
--- a/Engine/Source/Developer/MessageLog/Private/UserInterface/SMessageLogMessageListRow.cpp
+++ b/Engine/Source/Developer/MessageLog/Private/UserInterface/SMessageLogMessageListRow.cpp
@@ -86,7 +86,6 @@ TSharedRef<SWidget> SMessageLogMessageListRow::GenerateWidget()
 	}
 
 	return SNew(SHorizontalBox)
-		.ToolTipText(Message->ToText())
 
 		+ SHorizontalBox::Slot()
 		.AutoWidth()
@@ -307,6 +306,7 @@ void SMessageLogMessageListRow::CreateMessage(const TSharedRef<SHorizontalBox>&
 			[
 				SNew(STextBlock)
 				.ColorAndOpacity(FSlateColor::UseSubduedForeground())
+				.TextStyle(FEditorStyle::Get(), "Log.Normal")
 				.Text(FText::FromString(MessageString))
 			];
 		}
