Allow registering the Message Log Window from Insights Stand alone #rb Dave.Belanger

diff --git a/Engine/Source/Developer/MessageLog/Private/MessageLogModule.cpp b/Engine/Source/Developer/MessageLog/Private/MessageLogModule.cpp
--- a/Engine/Source/Developer/MessageLog/Private/MessageLogModule.cpp
+++ b/Engine/Source/Developer/MessageLog/Private/MessageLogModule.cpp
@@ -52,17 +52,12 @@ void FMessageLogModule::StartupModule()
 	MessageLogViewModel->Initialize();
 
 #if WITH_EDITOR
-	TWeakPtr<FMessageLogViewModel> WeakMessageLogViewModel = MessageLogViewModel;
 	ModulesChangedHandle = FModuleManager::Get().OnModulesChanged().AddLambda(
-	[WeakMessageLogViewModel](FName InModuleName, EModuleChangeReason InReason)
+	[this](FName InModuleName, EModuleChangeReason InReason)
 	{
 		if (InReason == EModuleChangeReason::ModuleLoaded && InModuleName == "LevelEditor")
 		{
-			FGlobalTabmanager::Get()->RegisterNomadTabSpawner("MessageLog", FOnSpawnTab::CreateStatic(&SpawnMessageLog, WeakMessageLogViewModel.Pin().ToSharedRef()))
-				.SetDisplayName(NSLOCTEXT("UnrealEditor", "MessageLogTab", "Message Log"))
-				.SetTooltipText(NSLOCTEXT("UnrealEditor", "MessageLogTooltipText", "Open the Message Log tab."))
-				.SetGroup(WorkspaceMenu::GetMenuStructure().GetDeveloperToolsLogCategory())
-				.SetIcon(FSlateIcon(FEditorStyle::GetStyleSetName(), "MessageLog.TabIcon"));
+			this->RegisterMessageLogSpawner(WorkspaceMenu::GetMenuStructure().GetDeveloperToolsLogCategory());
 		}
 	});
 #endif
@@ -133,3 +128,12 @@ void FMessageLogModule::EnableMessageLogDisplay(bool bInCanDisplayMessageLog)
 {
 	bCanDisplayMessageLog = bInCanDisplayMessageLog;
 }
+
+void FMessageLogModule::RegisterMessageLogSpawner(const TSharedRef<FWorkspaceItem>& InGroup)
+{
+	FGlobalTabmanager::Get()->RegisterNomadTabSpawner("MessageLog", FOnSpawnTab::CreateStatic(&SpawnMessageLog, MessageLogViewModel.ToSharedRef()))
+		.SetDisplayName(NSLOCTEXT("UnrealEditor", "MessageLogTab", "Message Log"))
+		.SetTooltipText(NSLOCTEXT("UnrealEditor", "MessageLogTooltipText", "Open the Message Log tab."))
+		.SetGroup(InGroup)
+		.SetIcon(FSlateIcon(FEditorStyle::GetStyleSetName(), "MessageLog.TabIcon"));
+}
diff --git a/Engine/Source/Developer/MessageLog/Public/MessageLogModule.h b/Engine/Source/Developer/MessageLog/Public/MessageLogModule.h
--- a/Engine/Source/Developer/MessageLog/Public/MessageLogModule.h
+++ b/Engine/Source/Developer/MessageLog/Public/MessageLogModule.h
@@ -7,6 +7,7 @@
 #include "MessageLogInitializationOptions.h"
 
 class FMessageLogViewModel;
+class FWorkspaceItem;
 
 DECLARE_DELEGATE_RetVal(bool, FCanShowMessageLog);
 
@@ -84,6 +85,12 @@ public:
 	 */
 	virtual void EnableMessageLogDisplay(bool bInCanDisplayMessageLog);
 
+	/**
+	* Registers the Message Log spawner.
+	* @param InGroup The Workspace group used to register the MessageLogSpawner.
+	 */
+	virtual void RegisterMessageLogSpawner(const TSharedRef<FWorkspaceItem>& InGroup);
+
 private:
 
 	/** The one and only message log view model */
