Fix some broken text layout in the message log

diff --git a/Engine/Source/Developer/MessageLog/Private/UserInterface/SMessageLog.cpp b/Engine/Source/Developer/MessageLog/Private/UserInterface/SMessageLog.cpp
--- a/Engine/Source/Developer/MessageLog/Private/UserInterface/SMessageLog.cpp
+++ b/Engine/Source/Developer/MessageLog/Private/UserInterface/SMessageLog.cpp
@@ -63,7 +63,7 @@ void SMessageLog::Construct( const FArguments& InArgs, const TSharedRef<FMessage
 			.Value(0.8f)
 			[
 				SAssignNew(CurrentListingDisplay, SBorder)
-					.BorderImage(FEditorStyle::GetBrush("NoBorder"))
+					.BorderImage(FEditorStyle::GetBrush("Brushes.Recessed"))
 					.Padding(0.0f)
 			]
 	];
diff --git a/Engine/Source/Developer/MessageLog/Private/UserInterface/SMessageLogMessageListRow.cpp b/Engine/Source/Developer/MessageLog/Private/UserInterface/SMessageLogMessageListRow.cpp
--- a/Engine/Source/Developer/MessageLog/Private/UserInterface/SMessageLogMessageListRow.cpp
+++ b/Engine/Source/Developer/MessageLog/Private/UserInterface/SMessageLogMessageListRow.cpp
@@ -279,6 +279,7 @@ void SMessageLogMessageListRow::CreateMessage(const TSharedRef<SHorizontalBox>&
 				SourceLink
 			]
 			+ SHorizontalBox::Slot()
+			.AutoWidth()
 			.Padding(0)
 			[
 				SNew(STextBlock)
@@ -326,6 +327,7 @@ void SMessageLogMessageListRow::CreateMessage(const TSharedRef<SHorizontalBox>&
 	if (RowContent.IsValid())
 	{
 		InHorzBox->AddSlot()
+			.AutoWidth()
 			.VAlign(VAlign_Center)
 			.Padding(Padding, 0.0f, 0.0f, 0.0f)
 			[
@@ -344,6 +346,7 @@ void SMessageLogMessageListRow::CreateMessage(const TSharedRef<SHorizontalBox>&
 				]
 
 				+ SHorizontalBox::Slot()
+				.AutoWidth()
 				.VAlign(VAlign_Center)
 				.Padding(2.0f, 0.0f, 0.0f, 0.0f)
 				[
