Fix missing well background color in message logs

diff --git a/Engine/Source/Developer/MessageLog/Private/UserInterface/SMessageLogListing.cpp b/Engine/Source/Developer/MessageLog/Private/UserInterface/SMessageLogListing.cpp
--- a/Engine/Source/Developer/MessageLog/Private/UserInterface/SMessageLogListing.cpp
+++ b/Engine/Source/Developer/MessageLog/Private/UserInterface/SMessageLogListing.cpp
@@ -43,35 +43,40 @@ void SMessageLogListing::Construct( const FArguments& InArgs, const TSharedRef<
 
 	ChildSlot
 	[
-		SAssignNew(VerticalBox, SVerticalBox)
-		+ SVerticalBox::Slot()
-		.FillHeight(1.0f)
-		.Padding(2.0f)
+		SNew(SBorder)
+		.BorderImage(FAppStyle::Get().GetBrush("Brushes.Recessed"))
+		.Padding(0)
 		[
-			SNew(SHorizontalBox)
-			+SHorizontalBox::Slot()
-			.FillWidth(1)
+			SAssignNew(VerticalBox, SVerticalBox)
+			+ SVerticalBox::Slot()
+			.FillHeight(1.0f)
+			.Padding(2.0f)
 			[
-				SNew(SScrollBox)
-				.Orientation(EOrientation::Orient_Horizontal)
-				+ SScrollBox::Slot()
+				SNew(SHorizontalBox)
+				+SHorizontalBox::Slot()
+				.FillWidth(1)
 				[
-					SAssignNew(MessageListView, SListView<TSharedRef<FTokenizedMessage>>)
-					.ListItemsSource(&MessageLogListingViewModel->GetFilteredMessages())
-					.OnGenerateRow(this, &SMessageLogListing::MakeMessageLogListItemWidget)
-					.OnSelectionChanged(this, &SMessageLogListing::OnLineSelectionChanged)
-					.ExternalScrollbar(ScrollBar)
-					.ItemHeight(24.0f)
-					.ConsumeMouseWheel(EConsumeMouseWheel::Always)
+					SNew(SScrollBox)
+					.Orientation(EOrientation::Orient_Horizontal)
+					+ SScrollBox::Slot()
+					[
+						SAssignNew(MessageListView, SListView<TSharedRef<FTokenizedMessage>>)
+						.ListItemsSource(&MessageLogListingViewModel->GetFilteredMessages())
+						.OnGenerateRow(this, &SMessageLogListing::MakeMessageLogListItemWidget)
+						.OnSelectionChanged(this, &SMessageLogListing::OnLineSelectionChanged)
+						.ExternalScrollbar(ScrollBar)
+						.ItemHeight(24.0f)
+						.ConsumeMouseWheel(EConsumeMouseWheel::Always)
+					]
 				]
-			]
-			+SHorizontalBox::Slot()
-			.AutoWidth()
-			[
-				SNew(SBox)
-				.WidthOverride(FOptionalSize(16))
+				+SHorizontalBox::Slot()
+				.AutoWidth()
 				[
-					ScrollBar
+					SNew(SBox)
+					.WidthOverride(FOptionalSize(16))
+					[
+						ScrollBar
+					]
 				]
 			]
 		]
