Fix notification for PIE errors never going away and not dismissable. Follow up to

diff --git a/Engine/Source/Developer/MessageLog/Private/Presentation/MessageLogListingViewModel.cpp b/Engine/Source/Developer/MessageLog/Private/Presentation/MessageLogListingViewModel.cpp
--- a/Engine/Source/Developer/MessageLog/Private/Presentation/MessageLogListingViewModel.cpp
+++ b/Engine/Source/Developer/MessageLog/Private/Presentation/MessageLogListingViewModel.cpp
@@ -315,12 +315,12 @@ void FMessageLogListingViewModel::NotifyIfAnyMessages( const FText& Message, EMe
 		{
 			int32 NotificationId = NextNotificationId++;
 
-			ErrorNotification.bFireAndForget = false;
+			ErrorNotification.bFireAndForget = true;
 			ErrorNotification.bUseThrobber = false;
 			ErrorNotification.FadeOutDuration = 0.f;
-			ErrorNotification.ExpireDuration = 0.f;
-			ErrorNotification.ButtonDetails.Add(FNotificationButtonInfo(LOCTEXT("DismissMessageButton", "Dismiss"), FText(), FSimpleDelegate::CreateSP(this, &FMessageLogListingViewModel::DismissNotification, NotificationId)));
-			ErrorNotification.ButtonDetails.Add(FNotificationButtonInfo(LOCTEXT("ShowMessageLogButton", "Show Message Log"), FText(), FSimpleDelegate::CreateSP(this, &FMessageLogListingViewModel::OpenMessageLogFromNotification, NotificationId)));
+			ErrorNotification.ExpireDuration = 15.f;
+			ErrorNotification.ButtonDetails.Add(FNotificationButtonInfo(LOCTEXT("DismissMessageButton", "Dismiss"), FText(), FSimpleDelegate::CreateSP(this, &FMessageLogListingViewModel::DismissNotification, NotificationId), SNotificationItem::CS_Fail));
+			ErrorNotification.ButtonDetails.Add(FNotificationButtonInfo(LOCTEXT("ShowMessageLogButton", "Show Message Log"), FText(), FSimpleDelegate::CreateSP(this, &FMessageLogListingViewModel::OpenMessageLogFromNotification, NotificationId), SNotificationItem::CS_Fail));
 
 			TSharedPtr<SNotificationItem> NewNotificationItem = FSlateNotificationManager::Get().AddNotification(ErrorNotification);
 			if (NewNotificationItem.IsValid())
