Updated Message Log to the new style guidelines and added a console command for testing Message Log

#okfforgithub public

diff --git a/Engine/Source/Developer/MessageLog/Private/MessageTesting.cpp b/Engine/Source/Developer/MessageLog/Private/MessageTesting.cpp
new file mode 100644
--- /dev/null
+++ b/Engine/Source/Developer/MessageLog/Private/MessageTesting.cpp
@@ -0,0 +1,15 @@
+// Copyright Epic Games, Inc. All Rights Reserved.
+
+#include "CoreMinimal.h"
+#include "HAL/IConsoleManager.h"
+#include "Logging/MessageLog.h" 
+
+static void TestMessages()
+{
+	FMessageLog PIELogger = FMessageLog(FName("SlateStyleLog"));
+	PIELogger.Message(EMessageSeverity::Type::Warning, FText::FromString("This is a Warning Message"));
+	PIELogger.Message(EMessageSeverity::Type::Error, FText::FromString("This is an Error Message"));
+	PIELogger.Message(EMessageSeverity::Type::Info, FText::FromString("This is an Info Message"));
+}
+
+FAutoConsoleCommand TestMessageLogCommand(TEXT("Slate.TestMessageLog"), TEXT(""), FConsoleCommandDelegate::CreateStatic(&TestMessages));
\ No newline at end of file
diff --git a/Engine/Source/Developer/MessageLog/Private/Presentation/MessageLogListingViewModel.cpp b/Engine/Source/Developer/MessageLog/Private/Presentation/MessageLogListingViewModel.cpp
--- a/Engine/Source/Developer/MessageLog/Private/Presentation/MessageLogListingViewModel.cpp
+++ b/Engine/Source/Developer/MessageLog/Private/Presentation/MessageLogListingViewModel.cpp
@@ -34,11 +34,11 @@ void FMessageLogListingViewModel::Initialize()
 	MessageLogListingModel->OnChanged().AddSP(this, &FMessageLogListingViewModel::OnChanged);
 
 	// Create our filters
-	MessageFilters.Add(MakeShareable(new FMessageFilter(LOCTEXT("CriticalErrors", "Critical Errors"), FSlateIcon("EditorStyle", "MessageLog.Error"))));
-	MessageFilters.Add(MakeShareable(new FMessageFilter(LOCTEXT("Errors", "Errors"), FSlateIcon("EditorStyle", "MessageLog.Error"))));
-	MessageFilters.Add(MakeShareable(new FMessageFilter(LOCTEXT("PerformanceWarnings", "Performance Warnings"), FSlateIcon("EditorStyle", "MessageLog.Warning"))));
-	MessageFilters.Add(MakeShareable(new FMessageFilter(LOCTEXT("Warnings", "Warnings"), FSlateIcon("EditorStyle", "MessageLog.Warning"))));
-	MessageFilters.Add(MakeShareable(new FMessageFilter(LOCTEXT("Info", "Info"), FSlateIcon("EditorStyle", "MessageLog.Note"))));
+	MessageFilters.Add(MakeShareable(new FMessageFilter(LOCTEXT("CriticalErrors", "Critical Errors"), FSlateIcon(FAppStyle::GetAppStyleSetName(), "Icons.ErrorWithColor"))));
+	MessageFilters.Add(MakeShareable(new FMessageFilter(LOCTEXT("Errors", "Errors"), FSlateIcon(FAppStyle::GetAppStyleSetName(), "Icons.ErrorWithColor"))));
+	MessageFilters.Add(MakeShareable(new FMessageFilter(LOCTEXT("PerformanceWarnings", "Performance Warnings"), FSlateIcon(FAppStyle::GetAppStyleSetName(), "Icons.WarningWithColor"))));
+	MessageFilters.Add(MakeShareable(new FMessageFilter(LOCTEXT("Warnings", "Warnings"), FSlateIcon(FAppStyle::GetAppStyleSetName(), "Icons.WarningWithColor"))));
+	MessageFilters.Add(MakeShareable(new FMessageFilter(LOCTEXT("Info", "Info"), FSlateIcon(FAppStyle::GetAppStyleSetName(), "Icons.BulletPoint"))));
 
 	for( auto It = MessageFilters.CreateConstIterator(); It; ++It)
 	{
diff --git a/Engine/Source/Developer/MessageLog/Private/UserInterface/SMessageLog.cpp b/Engine/Source/Developer/MessageLog/Private/UserInterface/SMessageLog.cpp
--- a/Engine/Source/Developer/MessageLog/Private/UserInterface/SMessageLog.cpp
+++ b/Engine/Source/Developer/MessageLog/Private/UserInterface/SMessageLog.cpp
@@ -30,42 +30,50 @@ void SMessageLog::Construct( const FArguments& InArgs, const TSharedRef<FMessage
 
 	ChildSlot
 	[
-		SNew(SSplitter)
-			.Orientation(Orient_Horizontal)
-
-		+ SSplitter::Slot()
-			.Value(0.2f)
-			[
-				SNew(SBorder)
-					.BorderImage(FEditorStyle::GetBrush("ToolPanel.GroupBorder"))
-					.Padding(0.0f)
-					[
-						// log categories list
-						SAssignNew(CategoriesListView, SListView<IMessageLogListingPtr>)
-							.ItemHeight(24.0f)
-							.ListItemsSource(&ViewModel->GetLogListingViewModels())
-							.OnGenerateRow(this, &SMessageLog::HandleCategoriesListGenerateRow)
-							.OnSelectionChanged(this, &SMessageLog::HandleCategoriesListSelectionChanged)
-							.SelectionMode(ESelectionMode::Single)
-							.HeaderRow
-							(
-								SNew(SHeaderRow)
-									.Visibility(EVisibility::Collapsed)
-
-								+ SHeaderRow::Column("Name")
-									.DefaultLabel(LOCTEXT("CategoriesListNameColumnHeader", "Category"))
-									.FillWidth(1.0)
-							)
-					]
-			]
-
-		+ SSplitter::Slot()
-			.Value(0.8f)
-			[
-				SAssignNew(CurrentListingDisplay, SBorder)
-					.BorderImage(FEditorStyle::GetBrush("Brushes.Recessed"))
-					.Padding(0.0f)
-			]
+		SNew(SBorder)
+		.BorderImage(FAppStyle::Get().GetBrush("Brushes.Panel"))
+		.Padding(FMargin(20.0f, 13.97f, 16.f, 16.f))
+		[
+			SNew(SSplitter)
+				.Orientation(Orient_Horizontal)
+				.Style(FAppStyle::Get(), "SplitterPanel")
+				.PhysicalSplitterHandleSize(10.0f)
+
+			+ SSplitter::Slot()
+				.Value(0.28f)
+				[
+					SNew(SBorder)
+						.BorderImage(FAppStyle::Get().GetBrush("Brushes.Recessed"))
+						.Padding(FMargin(16.f, 8.f, 0, 0))
+						[
+							// log categories list
+							SAssignNew(CategoriesListView, SListView<IMessageLogListingPtr>)
+								.ItemHeight(24.0f)
+								.ListItemsSource(&ViewModel->GetLogListingViewModels())
+								.OnGenerateRow(this, &SMessageLog::HandleCategoriesListGenerateRow)
+								.OnSelectionChanged(this, &SMessageLog::HandleCategoriesListSelectionChanged)
+								.SelectionMode(ESelectionMode::Single)
+								.HeaderRow
+								(
+									SNew(SHeaderRow)
+										.Visibility(EVisibility::Collapsed)
+
+									+ SHeaderRow::Column("Name")
+										.DefaultLabel(LOCTEXT("CategoriesListNameColumnHeader", "Category"))
+										.FillWidth(1.0)
+								)
+						]
+				]
+
+			+ SSplitter::Slot()
+				.Value(0.72f)
+				[
+					SAssignNew(CurrentListingDisplay, SBorder)
+						.BorderImage(FAppStyle::Get().GetBrush("Brushes.Recessed"))
+						.Padding(0.0f)
+				]
+		]
+
 	];
 
 	// Default to the first available (or the last used) listing
diff --git a/Engine/Source/Developer/MessageLog/Private/UserInterface/SMessageLogListing.cpp b/Engine/Source/Developer/MessageLog/Private/UserInterface/SMessageLogListing.cpp
--- a/Engine/Source/Developer/MessageLog/Private/UserInterface/SMessageLogListing.cpp
+++ b/Engine/Source/Developer/MessageLog/Private/UserInterface/SMessageLogListing.cpp
@@ -14,6 +14,7 @@
 #include "UserInterface/SMessageLogMessageListRow.h"
 #include "Framework/Commands/GenericCommands.h"
 #include "HAL/PlatformApplicationMisc.h"
+#include "Styling/StyleColors.h"
 
 
 #define LOCTEXT_NAMESPACE "Developer.MessageLog"
@@ -45,12 +46,12 @@ void SMessageLogListing::Construct( const FArguments& InArgs, const TSharedRef<
 	[
 		SNew(SBorder)
 		.BorderImage(FAppStyle::Get().GetBrush("Brushes.Recessed"))
-		.Padding(0)
+		.Padding(15.0f)
 		[
 			SAssignNew(VerticalBox, SVerticalBox)
 			+ SVerticalBox::Slot()
 			.FillHeight(1.0f)
-			.Padding(2.0f)
+			.Padding(0.f)
 			[
 				SNew(SHorizontalBox)
 				+SHorizontalBox::Slot()
@@ -89,7 +90,7 @@ void SMessageLogListing::Construct( const FArguments& InArgs, const TSharedRef<
 	{
 		VerticalBox->AddSlot()
 			.AutoHeight()
-			.Padding(6.0f)
+			.Padding(FMargin(6, 6, 0, 0))
 			[
 				SAssignNew(HorizontalBox, SHorizontalBox)
 			];
@@ -101,14 +102,13 @@ void SMessageLogListing::Construct( const FArguments& InArgs, const TSharedRef<
 				.HAlign(HAlign_Left)
 				[
 					SNew(SComboButton)
-					.ButtonStyle(FEditorStyle::Get(), "HoverHintOnly")
-					.ForegroundColor(FSlateColor::UseForeground())
-					.ContentPadding(2)
+					.ButtonStyle(FAppStyle::Get(), "SecondaryButton")
+					.ForegroundColor(FStyleColors::White)
 					.OnGetMenuContent(this, &SMessageLogListing::OnGetFilterMenuContent)
 					.ButtonContent()
 					[
 						SNew(STextBlock)
-						.Text(LOCTEXT("Show", "Show"))
+						.Text(LOCTEXT("Show", "SHOW"))
 						.ToolTipText(LOCTEXT("ShowToolTip", "Only show messages of the selected types"))
 					]
 				];
@@ -121,10 +121,10 @@ void SMessageLogListing::Construct( const FArguments& InArgs, const TSharedRef<
 				SNew(SComboButton)
 				.IsEnabled(this, &SMessageLogListing::IsPageWidgetEnabled)
 				.Visibility(this, &SMessageLogListing::GetPageWidgetVisibility)
-				.ButtonStyle(FEditorStyle::Get(), "HoverHintOnly")
-				.ForegroundColor(FSlateColor::UseForeground())
-				.ContentPadding(2)
+				.ButtonStyle(FAppStyle::Get(), "SecondaryButton")
+				.ForegroundColor(FStyleColors::White)
 				.OnGetMenuContent(this, &SMessageLogListing::OnGetPageMenuContent)
+				.ButtonColorAndOpacity(FStyleColors::White)
 				.ButtonContent()
 				[
 					SNew(STextBlock)
@@ -142,12 +142,12 @@ void SMessageLogListing::Construct( const FArguments& InArgs, const TSharedRef<
 				.OnClicked(this, &SMessageLogListing::OnClear)
 				.IsEnabled(this, &SMessageLogListing::IsClearWidgetEnabled)
 				.Visibility(this, &SMessageLogListing::GetClearWidgetVisibility)
-				.ButtonStyle(FEditorStyle::Get(), "HoverHintOnly")
-				.ForegroundColor(FSlateColor::UseForeground())
-				.ContentPadding(2)
+				.ButtonStyle(FAppStyle::Get(), "SecondaryButton")
+				.ForegroundColor(FStyleColors::White)
+				.ContentPadding(FMargin(6.5, 2))
 				[
 					SNew(STextBlock)
-					.Text(LOCTEXT("ClearMessageLog", "Clear"))
+					.Text(LOCTEXT("ClearMessageLog", "CLEAR"))
 					.ToolTipText(LOCTEXT("ClearMessageLog_ToolTip", "Clear the messages in this log"))
 				]
 			];
@@ -367,8 +367,15 @@ TSharedRef<ITableRow> SMessageLogListing::MakeShowWidget(TSharedRef<FMessageFilt
 				SNew(SHorizontalBox)
 				+SHorizontalBox::Slot().AutoWidth()
 				[
-					SNew(SImage)
-					.Image(Selection->GetIcon().GetIcon())
+					SNew(SBox)
+					.HAlign(HAlign_Center)
+					.VAlign(VAlign_Center)
+					.WidthOverride(16)
+					.HeightOverride(16)
+					[
+						SNew(SImage)
+						.Image(Selection->GetIcon().GetIcon())
+					]
 				]
 				+SHorizontalBox::Slot().AutoWidth()
 				[
@@ -398,7 +405,7 @@ FText SMessageLogListing::OnGetPageMenuLabel() const
 	}
 	else
 	{
-		return LOCTEXT("PageMenuLabel", "Page");
+		return LOCTEXT("PageMenuLabel", "PAGE");
 	}
 }
 
diff --git a/Engine/Source/Developer/MessageLog/Private/UserInterface/SMessageLogMessageListRow.cpp b/Engine/Source/Developer/MessageLog/Private/UserInterface/SMessageLogMessageListRow.cpp
--- a/Engine/Source/Developer/MessageLog/Private/UserInterface/SMessageLogMessageListRow.cpp
+++ b/Engine/Source/Developer/MessageLog/Private/UserInterface/SMessageLogMessageListRow.cpp
@@ -11,6 +11,7 @@
 #include "Misc/UObjectToken.h"
 #include "Widgets/Input/SHyperlink.h"
 #include "Internationalization/Regex.h"
+#include "Styling/StyleColors.h"
 
 #define LOCTEXT_NAMESPACE "SMessageLogMessageListRow"
 
@@ -22,7 +23,7 @@ TSharedRef<SWidget> SMessageLogMessageListRow::CreateHyperlink(const TSharedRef<
 	return SNew(SHyperlink)
 		.Text(InMessageToken->ToText())
 		.ToolTipText(InToolTip)
-		.TextStyle(FEditorStyle::Get(), "MessageLog")
+		.TextStyle(FAppStyle::Get(), "MessageLog")
 		.OnNavigate(this, &SMessageLogMessageListRow::HandleHyperlinkNavigate, InMessageToken);
 }
 
@@ -96,8 +97,16 @@ TSharedRef<SWidget> SMessageLogMessageListRow::GenerateWidget()
 			[
 				(SeverityImageName == NAME_None)
 				? SNullWidget::NullWidget
-				: static_cast<TSharedRef<SWidget>>(SNew(SImage)
-				.Image(FEditorStyle::GetBrush(SeverityImageName)))
+				: static_cast<TSharedRef<SWidget>>(
+					SNew(SBox)
+					.HAlign(HAlign_Center)
+					.VAlign(VAlign_Center)
+					.WidthOverride(16)
+					.HeightOverride(16)
+					[
+						SNew(SImage)
+						.Image(FAppStyle::Get().GetBrush(SeverityImageName))
+					])
 			]
 		]
 
@@ -117,7 +126,7 @@ TSharedRef<SWidget> SMessageLogMessageListRow::GenerateWidget()
 		!HasLinks
 		? SNullWidget::NullWidget
 		: static_cast<TSharedRef<SWidget>>(SNew(SBorder)
-		.BorderImage(FEditorStyle::GetBrush("ToolPanel.GroupBorder"))
+		.BorderImage(FAppStyle::Get().GetBrush("ToolPanel.GroupBorder"))
 		.Padding(FMargin(0.0f, 1.0f, 10.0f, 1.0f))
 		[
 			LinkBox
@@ -147,13 +156,13 @@ void SMessageLogMessageListRow::CreateMessage(const TSharedRef<SHorizontalBox>&
 					.Content()
 					[
 						SNew(SImage)
-						.Image(FEditorStyle::GetBrush(ImageToken->GetImageName()))
+						.Image(FAppStyle::Get().GetBrush(ImageToken->GetImageName()))
 					];
 			}
 			else
 			{
 				RowContent = SNew(SImage)
-					.Image(FEditorStyle::GetBrush(ImageToken->GetImageName()));
+					.Image(FAppStyle::Get().GetBrush(ImageToken->GetImageName()));
 			}
 		}
 	}
@@ -214,7 +223,7 @@ void SMessageLogMessageListRow::CreateMessage(const TSharedRef<SHorizontalBox>&
 		RowContent = SNew(SHyperlink)
 			.Text(InMessageToken->ToText())
 			.ToolTipText(ActionToken->GetActionDescription())
-			.TextStyle(FEditorStyle::Get(), "MessageLog")
+			.TextStyle(FAppStyle::Get(), "MessageLog")
 			.OnNavigate(this, &SMessageLogMessageListRow::HandleActionHyperlinkNavigate, ActionToken);
 
 		TokenContentVisbility = TAttribute<EVisibility>::Create(TAttribute<EVisibility>::FGetter::CreateRaw(this, &SMessageLogMessageListRow::GetActionLinkVisibility, ActionToken));
@@ -265,8 +274,8 @@ void SMessageLogMessageListRow::CreateMessage(const TSharedRef<SHorizontalBox>&
 				MessageString.RightChopInline(FileAndLineRegexMatcher.GetMatchEnding(), false);
 
 				SourceLink = SNew(SHyperlink)
-					.Style(FEditorStyle::Get(), "Common.GotoNativeCodeHyperlink")
-					.TextStyle(FEditorStyle::Get(), "MessageLog")
+					.Style(FAppStyle::Get(), "Common.GotoNativeCodeHyperlink")
+					.TextStyle(FAppStyle::Get(), "MessageLog")
 					.OnNavigate_Lambda([=] { FSlateApplication::Get().GotoLineInSource(FileName, LineNumber); })
 					.Text(FText::FromString(FileAndLineRegexMatcher.GetCaptureGroup(0)));
 			}
@@ -280,11 +289,11 @@ void SMessageLogMessageListRow::CreateMessage(const TSharedRef<SHorizontalBox>&
 			]
 			+ SHorizontalBox::Slot()
 			.AutoWidth()
-			.Padding(0)
+			.Padding(FMargin(0.f, 4.f))
 			[
 				SNew(STextBlock)
-				.TextStyle(FEditorStyle::Get(), "Log.Normal")
 				.Text(FText::FromString(MessageString))
+				.ColorAndOpacity(FStyleColors::Foreground)
 			];
 		}
 	}
@@ -304,7 +313,7 @@ void SMessageLogMessageListRow::CreateMessage(const TSharedRef<SHorizontalBox>&
 			DocumentationToken->GetPreviewExcerptLink(),
 			DocumentationToken->GetPreviewExcerptName())
 			)
-			.TextStyle(FEditorStyle::Get(), "MessageLog")
+			.TextStyle(FAppStyle::Get(), "MessageLog")
 			.OnNavigate(this, &SMessageLogMessageListRow::HandleDocsHyperlinkNavigate, DocumentationToken->GetDocumentationLink());
 	}
 		break;
@@ -317,7 +326,7 @@ void SMessageLogMessageListRow::CreateMessage(const TSharedRef<SHorizontalBox>&
 		RowContent = SNew(SHyperlink)
 			.Text(LOCTEXT("TutorialLabel", "Tutorial"))
 			.ToolTipText(LOCTEXT("TutorialTokenToolTip", "Click to open tutorial"))
-			.TextStyle(FEditorStyle::Get(), "MessageLog")
+			.TextStyle(FAppStyle::Get(), "MessageLog")
 			.OnNavigate(this, &SMessageLogMessageListRow::HandleTutorialHyperlinkNavigate, TutorialToken->GetTutorialAssetName());
 	}
 		break;
@@ -342,7 +351,7 @@ void SMessageLogMessageListRow::CreateMessage(const TSharedRef<SHorizontalBox>&
 					? SNullWidget::NullWidget
 					: static_cast<TSharedRef<SWidget>>(SNew(SImage)
 					.ColorAndOpacity(FSlateColor::UseForeground())
-					.Image(FEditorStyle::GetBrush(IconBrushName)))
+					.Image(FAppStyle::Get().GetBrush(IconBrushName)))
 				]
 
 				+ SHorizontalBox::Slot()
diff --git a/Engine/Source/Runtime/Core/Private/Logging/TokenizedMessage.cpp b/Engine/Source/Runtime/Core/Private/Logging/TokenizedMessage.cpp
--- a/Engine/Source/Runtime/Core/Private/Logging/TokenizedMessage.cpp
+++ b/Engine/Source/Runtime/Core/Private/Logging/TokenizedMessage.cpp
@@ -81,19 +81,19 @@ FName FTokenizedMessage::GetSeverityIconName(EMessageSeverity::Type InSeverity)
 	switch (InSeverity)
 	{
 	case EMessageSeverity::CriticalError:		
-		SeverityIconName = "MessageLog.Error";		
+		SeverityIconName = "Icons.ErrorWithColor";		
 		break;
 	case EMessageSeverity::Error:				
-		SeverityIconName = "MessageLog.Error";		
+		SeverityIconName = "Icons.ErrorWithColor";		
 		break;
 	case EMessageSeverity::PerformanceWarning:	
-		SeverityIconName = "MessageLog.Warning";	
+		SeverityIconName = "Icons.WarningWithColor";	
 		break;
 	case EMessageSeverity::Warning:				
-		SeverityIconName = "MessageLog.Warning";	
+		SeverityIconName = "Icons.WarningWithColor";	
 		break;
 	case EMessageSeverity::Info:				
-		SeverityIconName = "MessageLog.Note";		
+		SeverityIconName = "Icons.BulletPoint";		
 		break;
 	default:		
 		/* No icon for this type */						
diff --git a/Engine/Source/Runtime/SlateCore/Private/Styling/StarshipCoreStyle.cpp b/Engine/Source/Runtime/SlateCore/Private/Styling/StarshipCoreStyle.cpp
--- a/Engine/Source/Runtime/SlateCore/Private/Styling/StarshipCoreStyle.cpp
+++ b/Engine/Source/Runtime/SlateCore/Private/Styling/StarshipCoreStyle.cpp
@@ -380,6 +380,7 @@ TSharedRef<ISlateStyle> FStarshipCoreStyle::Create()
 		Style->Set("Icons.Link", new IMAGE_BRUSH_SVG("Starship/Common/Linked", Icon16x16));
 		Style->Set("Icons.Unlink", new IMAGE_BRUSH_SVG("Starship/Common/Unlinked", Icon16x16));
 
+		Style->Set("Icons.BulletPoint", new IMAGE_BRUSH_SVG("Starship/Common/bullet-point", Icon8x8));
 	}
 
 	// Tool panels
@@ -575,6 +576,10 @@ TSharedRef<ISlateStyle> FStarshipCoreStyle::Create()
 			.SetHandleNormalBrush(FSlateColorBrush(FStyleColors::Background))
 			.SetHandleHighlightBrush(FSlateColorBrush(FStyleColors::Secondary))
 		);
+		Style->Set("SplitterPanel", FSplitterStyle()
+			.SetHandleNormalBrush(FSlateColorBrush(FStyleColors::Panel))
+			.SetHandleHighlightBrush(FSlateColorBrush(FLinearColor(FColor(96, 96, 96))))
+		);
 	}
 
 	// TableView defaults...
