Control Rig: Python Message Log adds clear button

diff --git a/Engine/Source/Developer/MessageLog/Private/UserInterface/SMessageLogListing.cpp b/Engine/Source/Developer/MessageLog/Private/UserInterface/SMessageLogListing.cpp
--- a/Engine/Source/Developer/MessageLog/Private/UserInterface/SMessageLogListing.cpp
+++ b/Engine/Source/Developer/MessageLog/Private/UserInterface/SMessageLogListing.cpp
@@ -133,7 +133,6 @@ void SMessageLogListing::Construct( const FArguments& InArgs, const TSharedRef<
 				]
 			];
 
-		// if we aren't using pages, we allow the user to manually clear the log
 		HorizontalBox->AddSlot()
 			.FillWidth(1.0f)
 			.HAlign(HAlign_Right)
@@ -466,7 +465,7 @@ bool SMessageLogListing::IsClearWidgetEnabled() const
 
 EVisibility SMessageLogListing::GetClearWidgetVisibility() const
 {
-	if (MessageLogListingViewModel->GetAllowClear() && !MessageLogListingViewModel->GetShowPages())
+	if (MessageLogListingViewModel->GetAllowClear())
 	{
 		return EVisibility::Visible;
 	}
