Control Rig: Python Message Log scrolls to bottom when new messages are added

diff --git a/Engine/Source/Developer/MessageLog/Private/Presentation/MessageLogListingViewModel.cpp b/Engine/Source/Developer/MessageLog/Private/Presentation/MessageLogListingViewModel.cpp
--- a/Engine/Source/Developer/MessageLog/Private/Presentation/MessageLogListingViewModel.cpp
+++ b/Engine/Source/Developer/MessageLog/Private/Presentation/MessageLogListingViewModel.cpp
@@ -396,6 +396,16 @@ bool FMessageLogListingViewModel::GetShowPages() const
 	return bShowPages;
 }
 
+void FMessageLogListingViewModel::SetScrollToBottom(bool bInScrollToBottom)
+{
+	bScrollToBottom = bInScrollToBottom;
+}
+
+bool FMessageLogListingViewModel::GetScrollToBottom() const
+{
+	return bScrollToBottom;
+}
+
 void FMessageLogListingViewModel::SetAllowClear(bool bInAllowClear)
 {
 	bAllowClear = bInAllowClear;
diff --git a/Engine/Source/Developer/MessageLog/Private/Presentation/MessageLogListingViewModel.h b/Engine/Source/Developer/MessageLog/Private/Presentation/MessageLogListingViewModel.h
--- a/Engine/Source/Developer/MessageLog/Private/Presentation/MessageLogListingViewModel.h
+++ b/Engine/Source/Developer/MessageLog/Private/Presentation/MessageLogListingViewModel.h
@@ -106,6 +106,12 @@ public:
 	/** Get whether we should show pages or not */
 	bool GetShowPages() const;
 
+	/** Set whether we should scroll to the bottom when messages are added */
+	void SetScrollToBottom(bool bInScrollToBottom);
+
+	/** Get whether we should scroll to the bottom when messages are added */
+	bool GetScrollToBottom() const;
+
 	/** Set whether we should show allow the user to clear the log. */
 	void SetAllowClear(bool bInAllowClear);
 
@@ -157,6 +163,7 @@ private:
 		, bShowPages( InitializationOptions.bShowPages )
 		, bAllowClear( InitializationOptions.bAllowClear )
 		, bDiscardDuplicates( InitializationOptions.bDiscardDuplicates )
+		, bScrollToBottom( InitializationOptions.bScrollToBottom )
 		, MaxPageCount( InitializationOptions.MaxPageCount )
 		, bShowInLogWindow( InitializationOptions.bShowInLogWindow )
 		, CurrentPageIndex( 0 )
@@ -197,6 +204,9 @@ private:
 	/** Whether to check for duplicate messages & discard them */
 	bool bDiscardDuplicates;
 
+	/** Whether to scroll to the bottom when messages are added */
+	bool bScrollToBottom;
+
 	/** The limit on the number of displayed pages for this listing */
 	uint32 MaxPageCount;
 
diff --git a/Engine/Source/Developer/MessageLog/Private/UserInterface/SMessageLogListing.cpp b/Engine/Source/Developer/MessageLog/Private/UserInterface/SMessageLogListing.cpp
--- a/Engine/Source/Developer/MessageLog/Private/UserInterface/SMessageLogListing.cpp
+++ b/Engine/Source/Developer/MessageLog/Private/UserInterface/SMessageLogListing.cpp
@@ -197,9 +197,16 @@ void SMessageLogListing::OnSelectionChanged()
 void SMessageLogListing::RefreshVisibility()
 {
 	const TArray< TSharedRef<FTokenizedMessage> >& Messages = MessageLogListingViewModel->GetFilteredMessages();
-	if(Messages.Num() > 0)
+	if (Messages.Num() > 0)
 	{
-		ScrollToMessage( Messages[0] );
+		if (MessageLogListingViewModel->GetScrollToBottom())
+		{
+			ScrollToMessage( Messages.Last() );
+		}
+		else
+		{
+			ScrollToMessage( Messages[0] );
+		}
 	}
 	
 	MessageListView->RequestListRefresh();
diff --git a/Engine/Source/Developer/MessageLog/Public/MessageLogInitializationOptions.h b/Engine/Source/Developer/MessageLog/Public/MessageLogInitializationOptions.h
--- a/Engine/Source/Developer/MessageLog/Public/MessageLogInitializationOptions.h
+++ b/Engine/Source/Developer/MessageLog/Public/MessageLogInitializationOptions.h
@@ -16,6 +16,7 @@ struct FMessageLogInitializationOptions
 		, bDiscardDuplicates(false)
 		, MaxPageCount(20)
 		, bShowInLogWindow(true)
+		, bScrollToBottom(false)
 	{}
 
 	/** Whether to show the filters menu */
@@ -42,4 +43,7 @@ struct FMessageLogInitializationOptions
 
 	/** Whether to show this log in the main log window */
 	bool bShowInLogWindow;
+
+	/** Whether to scroll to the bottom of the window when messages are added */
+	bool bScrollToBottom;
 };
diff --git a/Engine/Source/Developer/RigVMDeveloper/Private/RigVMPythonUtils.cpp b/Engine/Source/Developer/RigVMDeveloper/Private/RigVMPythonUtils.cpp
--- a/Engine/Source/Developer/RigVMDeveloper/Private/RigVMPythonUtils.cpp
+++ b/Engine/Source/Developer/RigVMDeveloper/Private/RigVMPythonUtils.cpp
@@ -74,6 +74,7 @@ void RigVMPythonUtils::Print(const FString& InMessage)
 		InitOptions.bShowFilters = true;
 		InitOptions.bShowPages = true;
 		InitOptions.bAllowClear = true;
+		InitOptions.bScrollToBottom = true;
 		MessageLogModule.RegisterLogListing("ControlRigPythonLog", LOCTEXT("ControlRigPythonLog", "Control Rig Python Log"), InitOptions);
 	}
 	TSharedRef<IMessageLogListing> PythonLog = MessageLogModule.GetLogListing( TEXT("ControlRigPythonLog") );
