Fixed text style and highlight color in the MessageLog
#okfforgithub public

diff --git a/Engine/Source/Developer/MessageLog/Private/UserInterface/SMessageLogMessageListRow.cpp b/Engine/Source/Developer/MessageLog/Private/UserInterface/SMessageLogMessageListRow.cpp
--- a/Engine/Source/Developer/MessageLog/Private/UserInterface/SMessageLogMessageListRow.cpp
+++ b/Engine/Source/Developer/MessageLog/Private/UserInterface/SMessageLogMessageListRow.cpp
@@ -293,7 +293,8 @@ void SMessageLogMessageListRow::CreateMessage(const TSharedRef<SHorizontalBox>&
 			[
 				SNew(STextBlock)
 				.Text(FText::FromString(MessageString))
-				.ColorAndOpacity(FStyleColors::Foreground)
+				.ColorAndOpacity(FSlateColor::UseForeground())
+				.TextStyle(FAppStyle::Get(), "MessageLog")
 			];
 		}
 	}
diff --git a/Engine/Source/Runtime/SlateCore/Private/Styling/StarshipCoreStyle.cpp b/Engine/Source/Runtime/SlateCore/Private/Styling/StarshipCoreStyle.cpp
--- a/Engine/Source/Runtime/SlateCore/Private/Styling/StarshipCoreStyle.cpp
+++ b/Engine/Source/Runtime/SlateCore/Private/Styling/StarshipCoreStyle.cpp
@@ -769,7 +769,7 @@ TSharedRef<ISlateStyle> FStarshipCoreStyle::Create()
 	// Message Log
 	{
 		Style->Set("MessageLog", FTextBlockStyle(NormalText)
-			.SetFont(StyleFonts.Small)
+			.SetFont(DEFAULT_FONT("Mono", 9))
 			.SetShadowOffset(FVector2D::ZeroVector)
 		);
 		Style->Set("MessageLog.Error", new IMAGE_BRUSH("MessageLog/Log_Error", Icon16x16));
