Fix dismiss button on message log error notification not dismissing immediately

diff --git a/Engine/Source/Developer/MessageLog/Private/Presentation/MessageLogListingViewModel.cpp b/Engine/Source/Developer/MessageLog/Private/Presentation/MessageLogListingViewModel.cpp
--- a/Engine/Source/Developer/MessageLog/Private/Presentation/MessageLogListingViewModel.cpp
+++ b/Engine/Source/Developer/MessageLog/Private/Presentation/MessageLogListingViewModel.cpp
@@ -98,6 +98,7 @@ void FMessageLogListingViewModel::DismissNotification(int32 NotificationId)
 		TSharedPtr<SNotificationItem> NotificationItem = OpenNotifications[NotificationIndex].NotificationItem.Pin();
 		if (NotificationItem.IsValid())
 		{
+			NotificationItem->SetExpireDuration(0);
 			NotificationItem->ExpireAndFadeout();
 		}
 
