Change IMessagelogListing::GetSelectedMessagesAsString & GetAllMessagesAsString to return an FString. Previously it was using FTex to concatenate the strings which was causing OOM issues when attempting to copy several thousand lines from the message log to the clipboad.

diff --git a/Engine/Source/Developer/MessageLog/Private/Model/MessageLogListingModel.cpp b/Engine/Source/Developer/MessageLog/Private/Model/MessageLogListingModel.cpp
--- a/Engine/Source/Developer/MessageLog/Private/Model/MessageLogListingModel.cpp
+++ b/Engine/Source/Developer/MessageLog/Private/Model/MessageLogListingModel.cpp
@@ -22,22 +22,17 @@ const TSharedPtr< FTokenizedMessage >  FMessageLogListingModel::GetMessageAtInde
 	return FoundMessage;
 }
 
-FText FMessageLogListingModel::GetAllMessagesAsText( const uint32 PageIndex ) const
+FString FMessageLogListingModel::GetAllMessagesAsString( const uint32 PageIndex ) const
 {
-	FText CompiledText;
-
 	// Go through all the messages and add it to the compiled one.
 	FPage* Page = PageAtIndex(PageIndex);
+	TArray<FString> AllLines;
 	for( int32 MessageID = 0; MessageID < Page->Messages.Num(); MessageID++ )
 	{
-		const TSharedPtr< FTokenizedMessage > Message = Page->Messages[ MessageID ];
-		FFormatNamedArguments Args;
-		Args.Add( TEXT("PreviousMessage"), CompiledText );
-		Args.Add( TEXT("NewMessage"), Message.Get()->ToText() );
-		CompiledText = FText::Format( LOCTEXT("AggregateMessagesFormatter", "{PreviousMessage}{NewMessage}\n"), Args );
+		AllLines.Add(Page->Messages[ MessageID ]->ToText().ToString());
 	}
 
-	return CompiledText;
+	return FString::Join(AllLines, TEXT("\n"));
 }
 
 void FMessageLogListingModel::AddMessageInternal( const TSharedRef<FTokenizedMessage>& NewMessage, bool bMirrorToOutputLog, bool bDiscardDuplicates )
diff --git a/Engine/Source/Developer/MessageLog/Private/Model/MessageLogListingModel.h b/Engine/Source/Developer/MessageLog/Private/Model/MessageLogListingModel.h
--- a/Engine/Source/Developer/MessageLog/Private/Model/MessageLogListingModel.h
+++ b/Engine/Source/Developer/MessageLog/Private/Model/MessageLogListingModel.h
@@ -55,7 +55,7 @@ public:
 	const TSharedPtr<FTokenizedMessage> GetMessageAtIndex( const uint32 PageIndex, const int32 MessageIndex ) const;
 
 	/** Gets all messages as a string */
-	FText GetAllMessagesAsText(const uint32 PageIndex) const;
+	FString GetAllMessagesAsString(const uint32 PageIndex) const;
 
 	/** Obtains a const iterator to the message data structure */
 	MessageContainer::TConstIterator GetMessageIterator(uint32 PageIndex) const;
diff --git a/Engine/Source/Developer/MessageLog/Private/Presentation/MessageLogListingViewModel.cpp b/Engine/Source/Developer/MessageLog/Private/Presentation/MessageLogListingViewModel.cpp
--- a/Engine/Source/Developer/MessageLog/Private/Presentation/MessageLogListingViewModel.cpp
+++ b/Engine/Source/Developer/MessageLog/Private/Presentation/MessageLogListingViewModel.cpp
@@ -223,27 +223,22 @@ void FMessageLogListingViewModel::InvertSelectedMessages()
 	SelectionChangedEvent.Broadcast();
 }
 
-FText FMessageLogListingViewModel::GetSelectedMessagesAsText() const
+FString FMessageLogListingViewModel::GetSelectedMessagesAsString() const
 {
-	FText CompiledText;
-
 	// Go through each selected message and add it to the compiled one.
 	TArray< TSharedRef< FTokenizedMessage > > Selection = GetSelectedMessages();
+	TArray<FString> SelectedLines;
 	for( int32 MessageID = 0; MessageID < Selection.Num(); MessageID++ )
 	{
-		const TSharedPtr< FTokenizedMessage > Message = Selection[MessageID];
-		FFormatNamedArguments Args;
-		Args.Add( TEXT("PreviousMessage"), CompiledText );
-		Args.Add( TEXT("NewMessage"), Message.Get()->ToText() );
-		CompiledText = FText::Format( LOCTEXT("AggregateMessagesFormatter", "{PreviousMessage}{NewMessage}\n"), Args );
+		SelectedLines.Add(Selection[MessageID]->ToText().ToString());
 	}
 
-	return CompiledText;
+	return FString::Join(SelectedLines, TEXT("\n"));
 }
 
-FText FMessageLogListingViewModel::GetAllMessagesAsText() const
+FString FMessageLogListingViewModel::GetAllMessagesAsString() const
 {
-	return MessageLogListingModel->GetAllMessagesAsText(CurrentPageIndex);
+	return MessageLogListingModel->GetAllMessagesAsString(CurrentPageIndex);
 }
 
 const FName& FMessageLogListingViewModel::GetName() const
diff --git a/Engine/Source/Developer/MessageLog/Private/Presentation/MessageLogListingViewModel.h b/Engine/Source/Developer/MessageLog/Private/Presentation/MessageLogListingViewModel.h
--- a/Engine/Source/Developer/MessageLog/Private/Presentation/MessageLogListingViewModel.h
+++ b/Engine/Source/Developer/MessageLog/Private/Presentation/MessageLogListingViewModel.h
@@ -48,8 +48,8 @@ public:
 	virtual bool IsMessageSelected(const TSharedRef<class FTokenizedMessage>& Message) const override;
 	virtual void ClearSelectedMessages() override;
 	virtual void InvertSelectedMessages() override;
-	virtual FText GetSelectedMessagesAsText() const override;
-	virtual FText GetAllMessagesAsText() const override;
+	virtual FString GetSelectedMessagesAsString() const override;
+	virtual FString GetAllMessagesAsString() const override;
 	virtual const FName& GetName() const override;
 	virtual void SetLabel( const FText& InLogLabel ) override;
 	virtual const FText& GetLabel() const override;
diff --git a/Engine/Source/Developer/MessageLog/Private/UserInterface/SMessageLogListing.cpp b/Engine/Source/Developer/MessageLog/Private/UserInterface/SMessageLogListing.cpp
--- a/Engine/Source/Developer/MessageLog/Private/UserInterface/SMessageLogListing.cpp
+++ b/Engine/Source/Developer/MessageLog/Private/UserInterface/SMessageLogListing.cpp
@@ -273,15 +273,15 @@ void SMessageLogListing::InvertSelectedMessages() const
 }
 
 
-FText SMessageLogListing::GetSelectedMessagesAsText() const
+FString SMessageLogListing::GetSelectedMessagesAsString() const
 {
-	return MessageLogListingViewModel->GetSelectedMessagesAsText();
+	return MessageLogListingViewModel->GetSelectedMessagesAsString();
 }
 
 
-FText SMessageLogListing::GetAllMessagesAsText() const
+FString SMessageLogListing::GetAllMessagesAsString() const
 {
-	return MessageLogListingViewModel->GetAllMessagesAsText();
+	return MessageLogListingViewModel->GetAllMessagesAsString();
 }
 
 
@@ -312,28 +312,28 @@ void SMessageLogListing::OnLineSelectionChanged( TSharedPtr< FTokenizedMessage >
 
 void SMessageLogListing::CopySelectedToClipboard() const
 {
-	CopyText( true, true );
+	CopyLogOutput( true, true );
 }
 
 
-FText SMessageLogListing::CopyText( bool bSelected, bool bClipboard ) const
+FString SMessageLogListing::CopyLogOutput( bool bSelected, bool bClipboard ) const
 {
-	FText CombinedString;
+	FString CombinedString;
 
 	if( bSelected )
 	{
 		// Get the selected item and then get the selected messages as a string.
-		CombinedString = GetSelectedMessagesAsText();
+		CombinedString = GetSelectedMessagesAsString();
 	}
 	else
 	{
 		// Get the selected item and then get all the messages as a string.
-		CombinedString = GetAllMessagesAsText();
+		CombinedString = GetAllMessagesAsString();
 	}
 	if( bClipboard )
 	{
 		// Pass that to the clipboard.
-		FPlatformApplicationMisc::ClipboardCopy( *CombinedString.ToString() );
+		FPlatformApplicationMisc::ClipboardCopy( *CombinedString );
 	}
 
 	return CombinedString;
diff --git a/Engine/Source/Developer/MessageLog/Private/UserInterface/SMessageLogListing.h b/Engine/Source/Developer/MessageLog/Private/UserInterface/SMessageLogListing.h
--- a/Engine/Source/Developer/MessageLog/Private/UserInterface/SMessageLogListing.h
+++ b/Engine/Source/Developer/MessageLog/Private/UserInterface/SMessageLogListing.h
@@ -60,10 +60,10 @@ public:
 	void InvertSelectedMessages() const;
 
 	/** Compiles the selected messages into a single string. */
-	FText GetSelectedMessagesAsText() const;
+	FString GetSelectedMessagesAsString() const;
 
 	/** Compiles all the messages into a single string. */
-	FText GetAllMessagesAsText() const;
+	FString GetAllMessagesAsString() const;
 
 	/** Gets the message log listing unique name */
 	const FName& GetName() const { return MessageLogListingViewModel->GetName(); }
@@ -86,9 +86,9 @@ public:
 	void CopySelectedToClipboard() const;
 
 	/**
-	 * Copies text (selected/all), optionally to the clipboard
+	 * Copies log output (selected/all), optionally to the clipboard
 	 */
-	FText CopyText( bool bSelected, bool bClipboard ) const;
+	FString CopyLogOutput( bool bSelected, bool bClipboard ) const;
 
 	/** @return	The UICommandList supported by the MessageLogView */
 	const TSharedRef< const FUICommandList > GetCommandList() const;
diff --git a/Engine/Source/Developer/MessageLog/Public/IMessageLogListing.h b/Engine/Source/Developer/MessageLog/Public/IMessageLogListing.h
--- a/Engine/Source/Developer/MessageLog/Public/IMessageLogListing.h
+++ b/Engine/Source/Developer/MessageLog/Public/IMessageLogListing.h
@@ -73,14 +73,14 @@ public:
 	 *
 	 * @returns the selected messages as a string
 	 */
-	virtual FText GetSelectedMessagesAsText() const = 0;
+	virtual FString GetSelectedMessagesAsString() const = 0;
 
 	/** 
 	 * Gets all the unfiltered messages as a single piece of text
 	 *
 	 * @returns all the messages as a string
 	 */
-	virtual FText GetAllMessagesAsText() const = 0;
+	virtual FString GetAllMessagesAsString() const = 0;
 
 	/** 
 	 * Gets the message log listing unique name 
