If dismissing a toast created by the message log, set the expire duration to 0.

diff --git a/Engine/Source/Developer/MessageLog/Private/Presentation/MessageLogListingViewModel.cpp b/Engine/Source/Developer/MessageLog/Private/Presentation/MessageLogListingViewModel.cpp
--- a/Engine/Source/Developer/MessageLog/Private/Presentation/MessageLogListingViewModel.cpp
+++ b/Engine/Source/Developer/MessageLog/Private/Presentation/MessageLogListingViewModel.cpp
@@ -97,7 +97,7 @@ void FMessageLogListingViewModel::DismissNotification(int32 NotificationId)
 		TSharedPtr<SNotificationItem> NotificationItem = OpenNotifications[NotificationIndex].NotificationItem.Pin();
 		if (NotificationItem.IsValid())
 		{
-			NotificationItem->SetExpireDuration(0);
+			NotificationItem->SetExpireDuration(0.0f);
 			NotificationItem->ExpireAndFadeout();
 		}
 
