PR #9872: Dismiss identical warning notification popups. (Contributed by [REDACTED_PERSON])

diff --git a/Engine/Source/Developer/MessageLog/Private/Presentation/MessageLogListingViewModel.cpp b/Engine/Source/Developer/MessageLog/Private/Presentation/MessageLogListingViewModel.cpp
--- a/Engine/Source/Developer/MessageLog/Private/Presentation/MessageLogListingViewModel.cpp
+++ b/Engine/Source/Developer/MessageLog/Private/Presentation/MessageLogListingViewModel.cpp
@@ -84,7 +84,7 @@ const TSharedPtr<FTokenizedMessage>  FMessageLogListingViewModel::GetMessageAtIn
 	return FoundMessage;
 }
 
-const TArray< TSharedRef< class FMessageFilter> >& FMessageLogListingViewModel::GetMessageFilters() const
+const TArray< TSharedRef<FMessageFilter> >& FMessageLogListingViewModel::GetMessageFilters() const
 {
 	return MessageFilters;
 }
@@ -147,12 +147,12 @@ EMessageSeverity::Type FMessageLogListingViewModel::HighestSeverityPresent( uint
 	return Severity;
 }
 
-void FMessageLogListingViewModel::AddMessage( const TSharedRef< class FTokenizedMessage >& NewMessage, bool bMirrorToOutputLog )
+void FMessageLogListingViewModel::AddMessage( const TSharedRef< FTokenizedMessage >& NewMessage, bool bMirrorToOutputLog )
 {
 	MessageLogListingModel->AddMessage(NewMessage, bMirrorToOutputLog, bDiscardDuplicates);
 }
 
-void FMessageLogListingViewModel::AddMessages( const TArray< TSharedRef< class FTokenizedMessage > >& NewMessages, bool bMirrorToOutputLog )
+void FMessageLogListingViewModel::AddMessages( const TArray< TSharedRef< FTokenizedMessage > >& NewMessages, bool bMirrorToOutputLog )
 {
 	MessageLogListingModel->AddMessages(NewMessages, bMirrorToOutputLog, bDiscardDuplicates);
 }
@@ -170,17 +170,17 @@ void FMessageLogListingViewModel::SelectMessages( const TArray< TSharedRef<FToke
 	SelectionChangedEvent.Broadcast();
 }
 
-const TArray< TSharedRef<class FTokenizedMessage> >& FMessageLogListingViewModel::GetSelectedMessages() const
+const TArray< TSharedRef<FTokenizedMessage> >& FMessageLogListingViewModel::GetSelectedMessages() const
 {
 	return SelectedFilteredMessages;
 }
 
-const TArray< TSharedRef<class FTokenizedMessage> >& FMessageLogListingViewModel::GetFilteredMessages() const
+const TArray< TSharedRef<FTokenizedMessage> >& FMessageLogListingViewModel::GetFilteredMessages() const
 {
 	return FilteredMessages;
 }
 
-void FMessageLogListingViewModel::SelectMessage(const TSharedRef<class FTokenizedMessage>& Message, bool bSelected)
+void FMessageLogListingViewModel::SelectMessage(const TSharedRef<FTokenizedMessage>& Message, bool bSelected)
 {
 	bool bIsAlreadySelected = INDEX_NONE != SelectedFilteredMessages.Find( Message );
 
@@ -200,7 +200,7 @@ void FMessageLogListingViewModel::SelectMessage(const TSharedRef<class FTokenize
 	}
 }
 
-bool FMessageLogListingViewModel::IsMessageSelected(const TSharedRef<class FTokenizedMessage>& Message) const
+bool FMessageLogListingViewModel::IsMessageSelected(const TSharedRef<FTokenizedMessage>& Message) const
 {
 	return INDEX_NONE != SelectedFilteredMessages.Find(Message);
 }
@@ -255,7 +255,7 @@ const FText& FMessageLogListingViewModel::GetLabel() const
 	return LogLabel;
 }
 
-void FMessageLogListingViewModel::ExecuteToken( const TSharedRef<class IMessageToken>& Token ) const
+void FMessageLogListingViewModel::ExecuteToken( const TSharedRef<IMessageToken>& Token ) const
 {
 	TokenClickedEvent.Broadcast( Token );
 }
@@ -333,6 +333,16 @@ void FMessageLogListingViewModel::NotifyIfAnyMessages( const FText& Message, EMe
 
 		if (NumMessagesPresent(0, EMessageSeverity::Error) > 0)
 		{
+			// If there is a notification with the same message, dismiss it to make space for this new one.
+			const FOpenNotification* ExistingNotification = OpenNotifications.FindByPredicate([NotificationMessage](const FOpenNotification& Notification)
+			{
+				return Notification.NotificationMessage.CompareTo(NotificationMessage) == 0;
+			});
+			if(ExistingNotification != nullptr)
+			{
+				DismissNotification(ExistingNotification->NotificationId);
+			}
+			
 			int32 NotificationId = NextNotificationId++;
 
 			ErrorNotification.bFireAndForget = true;
@@ -346,7 +356,7 @@ void FMessageLogListingViewModel::NotifyIfAnyMessages( const FText& Message, EMe
 			if (NewNotificationItem.IsValid())
 			{
 				NewNotificationItem->SetCompletionState(SNotificationItem::CS_Fail);
-				OpenNotifications.Emplace(NotificationId, NewNotificationItem);
+				OpenNotifications.Emplace(NotificationId, NewNotificationItem, NotificationMessage);
 			}
 		}
 		else
diff --git a/Engine/Source/Developer/MessageLog/Private/Presentation/MessageLogListingViewModel.h b/Engine/Source/Developer/MessageLog/Private/Presentation/MessageLogListingViewModel.h
--- a/Engine/Source/Developer/MessageLog/Private/Presentation/MessageLogListingViewModel.h
+++ b/Engine/Source/Developer/MessageLog/Private/Presentation/MessageLogListingViewModel.h
@@ -253,7 +253,16 @@ private:
 	{
 		int32 NotificationId;
 		TWeakPtr<class SNotificationItem> NotificationItem;
-		FOpenNotification(int32 InNotificationId, const TWeakPtr<class SNotificationItem>& InNotificationItem) : NotificationId(InNotificationId), NotificationItem(InNotificationItem) {}
+		const FText NotificationMessage;
+
+		FOpenNotification(int32 InNotificationId,
+		                  const TWeakPtr<class SNotificationItem>& InNotificationItem,
+		                  const FText InNotificationMessage)
+			: NotificationId(InNotificationId)
+			, NotificationItem(InNotificationItem)
+			, NotificationMessage( InNotificationMessage)
+		{
+		}
 		FOpenNotification() : NotificationId(INDEX_NONE) {}
 	};
 	TArray<FOpenNotification> OpenNotifications;
