Minor dependency cleanup in MessageLog.

diff --git a/Engine/Source/Developer/MessageLog/Private/Presentation/MessageLogListingViewModel.cpp b/Engine/Source/Developer/MessageLog/Private/Presentation/MessageLogListingViewModel.cpp
--- a/Engine/Source/Developer/MessageLog/Private/Presentation/MessageLogListingViewModel.cpp
+++ b/Engine/Source/Developer/MessageLog/Private/Presentation/MessageLogListingViewModel.cpp
@@ -2,6 +2,7 @@
 
 #include "Presentation/MessageLogListingViewModel.h"
 #include "Modules/ModuleManager.h"
+#include "MessageFilter.h"
 #include "MessageLogModule.h"
 #include "Framework/Notifications/NotificationManager.h"
 #include "Widgets/Notifications/SNotificationList.h"
@@ -10,6 +11,8 @@
 
 int32 FMessageLogListingViewModel::NextNotificationId = 0;
 
+FMessageLogListingViewModel::FMessageLogListingViewModel() = default;
+
 FMessageLogListingViewModel::~FMessageLogListingViewModel()
 {
 	check(MessageLogListingModel.IsValid());
diff --git a/Engine/Source/Developer/MessageLog/Private/Presentation/MessageLogListingViewModel.h b/Engine/Source/Developer/MessageLog/Private/Presentation/MessageLogListingViewModel.h
--- a/Engine/Source/Developer/MessageLog/Private/Presentation/MessageLogListingViewModel.h
+++ b/Engine/Source/Developer/MessageLog/Private/Presentation/MessageLogListingViewModel.h
@@ -7,7 +7,8 @@
 #include "Logging/TokenizedMessage.h"
 #include "IMessageLogListing.h"
 #include "Model/MessageLogListingModel.h"
-#include "MessageFilter.h"
+
+class FMessageFilter;
 
 /**
  * The non-UI solution specific presentation logic for a collection of messages for a particular system.
@@ -34,7 +35,7 @@ public:
 	}
 
 public:
-
+	FMessageLogListingViewModel();
 	~FMessageLogListingViewModel();
 
 	/** Begin IMessageLogListing interface */
@@ -53,7 +54,7 @@ public:
 	virtual const FName& GetName() const override;
 	virtual void SetLabel( const FText& InLogLabel ) override;
 	virtual const FText& GetLabel() const override;
-	virtual const TArray< TSharedRef< class FMessageFilter> >& GetMessageFilters() const override;
+	virtual const TArray< TSharedRef<FMessageFilter> >& GetMessageFilters() const override;
 	virtual void ExecuteToken( const TSharedRef<class IMessageToken>& Token ) const override;
 	virtual void NewPage( const FText& Title ) override;
 	virtual void SetCurrentPage( const FText& Title ) override;
diff --git a/Engine/Source/Developer/MessageLog/Private/UserInterface/SMessageLogListing.cpp b/Engine/Source/Developer/MessageLog/Private/UserInterface/SMessageLogListing.cpp
--- a/Engine/Source/Developer/MessageLog/Private/UserInterface/SMessageLogListing.cpp
+++ b/Engine/Source/Developer/MessageLog/Private/UserInterface/SMessageLogListing.cpp
@@ -1,6 +1,7 @@
 // Copyright Epic Games, Inc. All Rights Reserved.
 
 #include "UserInterface/SMessageLogListing.h"
+#include "MessageFilter.h"
 #include "Widgets/Images/SImage.h"
 #include "Widgets/Text/STextBlock.h"
 #include "Framework/MultiBox/MultiBoxBuilder.h"
diff --git a/Engine/Source/Developer/MessageLog/Private/UserInterface/SMessageLogListing.h b/Engine/Source/Developer/MessageLog/Private/UserInterface/SMessageLogListing.h
--- a/Engine/Source/Developer/MessageLog/Private/UserInterface/SMessageLogListing.h
+++ b/Engine/Source/Developer/MessageLog/Private/UserInterface/SMessageLogListing.h
@@ -11,10 +11,11 @@
 #include "Widgets/SCompoundWidget.h"
 #include "Widgets/Views/STableViewBase.h"
 #include "Widgets/Views/STableRow.h"
-#include "MessageFilter.h"
 #include "Presentation/MessageLogListingViewModel.h"
 #include "Framework/Commands/UICommandList.h"
 
+class FMessageFilter;
+
 /**
  * A message log listing, such as the Compiler Log, or the Map Check Log.
  * Holds the log lines, and any extra widgets necessary.
