Moving a few MessageLog headers to Public since they're accessed outside of this module.

diff --git a/Engine/Source/Developer/MessageLog/Private/Model/MessageLogListingModel.h b/Engine/Source/Developer/MessageLog/Private/Model/MessageLogListingModel.h
deleted file mode 100644
diff --git a/Engine/Source/Developer/MessageLog/Private/Presentation/MessageLogListingViewModel.h b/Engine/Source/Developer/MessageLog/Private/Presentation/MessageLogListingViewModel.h
deleted file mode 100644
diff --git a/Engine/Source/Developer/MessageLog/Private/UserInterface/SMessageLogListing.h b/Engine/Source/Developer/MessageLog/Private/UserInterface/SMessageLogListing.h
deleted file mode 100644
diff --git a/Engine/Source/Developer/MessageLog/Private/UserInterface/SMessageLogMessageListRow.h b/Engine/Source/Developer/MessageLog/Private/UserInterface/SMessageLogMessageListRow.h
--- a/Engine/Source/Developer/MessageLog/Private/UserInterface/SMessageLogMessageListRow.h
+++ b/Engine/Source/Developer/MessageLog/Private/UserInterface/SMessageLogMessageListRow.h
@@ -20,7 +20,7 @@
 
 #define LOCTEXT_NAMESPACE "SMessageLogMessageListRow"
 
-class MESSAGELOG_API SMessageLogMessageListRow
+class SMessageLogMessageListRow
 	: public STableRow<TSharedPtr<FTokenizedMessage>>
 {
 public:
diff --git a/Engine/Source/Developer/MessageLog/Public/Model/MessageLogListingModel.h b/Engine/Source/Developer/MessageLog/Public/Model/MessageLogListingModel.h
new file mode 100644
--- /dev/null
+++ b/Engine/Source/Developer/MessageLog/Public/Model/MessageLogListingModel.h
@@ -0,0 +1,180 @@
+// Copyright Epic Games, Inc. All Rights Reserved.
+
+#pragma once
+
+#include "CoreMinimal.h"
+#include "Containers/List.h"
+#include "Logging/TokenizedMessage.h"
+
+/** Define the container type for all the messages */
+typedef TArray< TSharedRef<FTokenizedMessage> > MessageContainer;
+
+/** This class represents a set of rich tokenized messages for a particular system */
+class MESSAGELOG_API FMessageLogListingModel : public TSharedFromThis< FMessageLogListingModel >
+{
+protected:
+	struct FPage
+	{
+		FPage( const FText& InTitle )
+			: Title( InTitle )
+		{
+		}
+
+		/** The title of this page */
+		FText Title;
+
+		/** The list of messages in this log listing */
+		MessageContainer Messages;
+
+		TSet<uint32> MessagesHashes;
+	};
+
+public:
+
+	virtual ~FMessageLogListingModel() {}
+
+	/**  
+	 *	Factory method which creates a new FMessageLogListingModel object
+	 *
+	 *	@param	InLogName		The name of the log
+	 */
+	static TSharedRef< FMessageLogListingModel > Create( const FName& InLogName )
+	{
+		TSharedRef< FMessageLogListingModel > NewLogListing( new FMessageLogListingModel( InLogName ) );
+		return NewLogListing;
+	}
+
+	/** Broadcasts whenever the message log listing changes */
+	DECLARE_EVENT( FMessageLogListingModel, FChangedEvent )
+	FChangedEvent& OnChanged() { return ChangedEvent; }
+
+	/** Retrieves the name identifier for this log listing */
+	const FName& GetName() const { return LogName; }
+
+	/** Returns the message at the specified index */
+	const TSharedPtr<FTokenizedMessage> GetMessageAtIndex( const uint32 PageIndex, const int32 MessageIndex ) const;
+
+	/** Gets all messages as a string */
+	FString GetAllMessagesAsString(const uint32 PageIndex) const;
+
+	/** Obtains a const iterator to the message data structure */
+	MessageContainer::TConstIterator GetMessageIterator(uint32 PageIndex) const;
+
+	/** Replaces the message at the given index */
+	int32 ReplaceMessage( const TSharedRef< FTokenizedMessage >& NewMessage, const uint32 PageIndex, const int32 MessageIndex );
+
+	/** Appends a message */
+	void AddMessage( const TSharedRef< FTokenizedMessage >& NewMessage, bool bMirrorToOutputLog = true, bool bDiscardDuplicates = false );
+
+	/** Appends multiple messages */
+	void AddMessages( const TArray< TSharedRef< FTokenizedMessage > >& NewMessages, bool bMirrorToOutputLog = true, bool bDiscardDuplicates = false );
+
+	/** Clears all messages */
+	void ClearMessages();
+
+	/**
+	 * Add a new page. Old pages are only kept around if they contain messages, so if the current
+	 * page is empty, this call does nothing.
+	 * @param	InTitle		The title that will be displayed for this page.
+	 * @param	InMaxPages	The maximum number of pages we keep around. If the count is exceeded, 
+	 *						we discard the oldest page.
+	 */
+	void NewPage( const FText& InTitle, uint32 InMaxPages );
+
+	/**
+	* Sets the current page to the one specified by the title. If the page does not exist, creates a new page.
+	* @param	InTitle		The title of the page to search for, or create
+	* @param	InMaxPages	The maximum number of pages we keep around. If the count is exceeded, 
+	*						we discard the oldest page.
+	* @returns true if the page has changed (either switched or created a new one).
+	*/
+	bool SetCurrentPage( const FText& InTitle, uint32 InMaxPages  );
+
+	/**
+	* Sets the current page to the one specified by the index. If the page does not exist, creates a new page.
+	* The page will move to index 0.
+	* @param	InTitle		The title of the page to search for, or create
+	* @param	InMaxPages	The maximum number of pages we keep around. If the count is exceeded, 
+	*						we discard the oldest page.
+	* @returns true if the page has changed (either switched or created a new one).
+	*/
+	bool SetCurrentPage( const uint32 InOldPageIndex  );
+
+	/** Get the number of pages contained in this log */
+	uint32 NumPages() const;
+
+	/** Get the number of messages on the passed-in page */
+	uint32 NumMessages( uint32 PageIndex ) const;
+
+	/**
+	 * Get the title of the page at the specified index
+	 * @param	PageIndex	The index of the page
+	 */
+	const FText& GetPageTitle( const uint32 PageIndex ) const;
+
+	/** Helper function for RemoveDuplicates(), exposed so the ViewModel can use it too */
+	static bool AreMessagesEqual(const TSharedRef< FTokenizedMessage >& Message0, const TSharedRef< FTokenizedMessage >& Message1);
+
+	/** Remove any messages that are duplicates of one another - O(n) */
+	void RemoveDuplicates(uint32 PageIndex);
+
+protected:
+
+	/** Will broadcast to all registered observers informing them of a change */
+	virtual void Notify() { ChangedEvent.Broadcast(); }
+
+	/** Access the current page (we only add messages to this page */
+	FPage& CurrentPage() const;
+
+	/** Get a page by index - uses cache to speed up linked list access */
+	FPage* PageAtIndex(const uint32 PageIndex) const;
+
+	/** Create a new page if we have one pending */
+	void CreateNewPageIfRequired();
+
+private:
+	/**
+	 *	FMessageLogListingModel Constructor
+	 *
+	 *	@param	InLogName		The name of the log
+	 */
+	FMessageLogListingModel( const FName& InLogName )
+		: LogName( InLogName )
+		, bIsPrintingToOutputLog( false )
+	{
+		check( LogName != NAME_None );
+
+		// create default page
+		Pages.AddTail(FPage(FText::FromName(LogName)));
+		CachedPage = &Pages.GetTail()->GetValue();
+		CachedPageIndex = 0;
+	}
+
+	/** Helper function for AddMessage and AddMessages */
+	void AddMessageInternal( const TSharedRef<FTokenizedMessage>& NewMessage, bool bMirrorToOutputLog, bool bDiscardDuplicates );
+	
+private:
+	/** The name of a pending page */
+	FText PendingPageName;
+
+	/** The cap on the number of pages we have */
+	uint32 MaxPages;
+
+	/** The list of pages in this log listing */
+	TDoubleLinkedList<FPage> Pages;
+
+	/** Name of the listing, for identification */
+	FName LogName;
+
+	/** Delegate to call when data is changed */
+	FChangedEvent ChangedEvent;
+
+	// Are we currently processing the output log mirror?  If so, we drop any additional messages we receive, as they are duplicates
+	bool bIsPrintingToOutputLog;
+
+	/** Cached page index */
+	mutable uint32 CachedPageIndex;
+
+	/** Cached page */
+	mutable FPage* CachedPage;
+};
diff --git a/Engine/Source/Developer/MessageLog/Public/Presentation/MessageLogListingViewModel.h b/Engine/Source/Developer/MessageLog/Public/Presentation/MessageLogListingViewModel.h
new file mode 100644
--- /dev/null
+++ b/Engine/Source/Developer/MessageLog/Public/Presentation/MessageLogListingViewModel.h
@@ -0,0 +1,271 @@
+// Copyright Epic Games, Inc. All Rights Reserved.
+
+#pragma once
+
+#include "CoreMinimal.h"
+#include "MessageLogInitializationOptions.h"
+#include "Logging/TokenizedMessage.h"
+#include "IMessageLogListing.h"
+#include "Model/MessageLogListingModel.h"
+
+class FMessageFilter;
+
+/**
+ * The non-UI solution specific presentation logic for a collection of messages for a particular system.
+ */
+class MESSAGELOG_API FMessageLogListingViewModel
+	: public TSharedFromThis<FMessageLogListingViewModel>
+	, public IMessageLogListing
+{
+public:
+
+	/**  
+	 *	Factory method which creates a new FMessageLogListingViewModel object
+	 *
+	 *	@param	InMessageLogListingModel		The MessageLogListing data to view
+	 *	@param	InLogLabel						The label that will be displayed in the UI for this log listing
+	 *	@param	InMessageLogListingModel		If true, a filters list will be displayed for this log listing
+	 */
+	static TSharedRef< FMessageLogListingViewModel > Create( const TSharedRef< FMessageLogListingModel >& InMessageLogListingModel, const FText& InLogLabel, const FMessageLogInitializationOptions& InitializationOptions = FMessageLogInitializationOptions() )
+	{
+		TSharedRef< FMessageLogListingViewModel > NewLogListingView( new FMessageLogListingViewModel( InMessageLogListingModel, InLogLabel, InitializationOptions ) );
+		NewLogListingView->Initialize();
+
+		return NewLogListingView;
+	}
+
+public:
+	FMessageLogListingViewModel();
+	~FMessageLogListingViewModel();
+
+	/** Begin IMessageLogListing interface */
+	virtual void AddMessage( const TSharedRef< class FTokenizedMessage >& NewMessage, bool bMirrorToOutputLog ) override;
+	virtual void AddMessages( const TArray< TSharedRef< class FTokenizedMessage > >& NewMessages, bool bMirrorToOutputLog ) override;
+	virtual void ClearMessages() override;
+	virtual const TArray< TSharedRef<class FTokenizedMessage> >& GetSelectedMessages() const override;
+	virtual void SelectMessages( const TArray< TSharedRef<class FTokenizedMessage> >& InSelectedMessages ) override;
+	virtual const TArray< TSharedRef<class FTokenizedMessage> >& GetFilteredMessages() const override;
+	virtual void SelectMessage(const TSharedRef<class FTokenizedMessage>& Message, bool bSelected) override;
+	virtual bool IsMessageSelected(const TSharedRef<class FTokenizedMessage>& Message) const override;
+	virtual void ClearSelectedMessages() override;
+	virtual void InvertSelectedMessages() override;
+	virtual FString GetSelectedMessagesAsString() const override;
+	virtual FString GetAllMessagesAsString() const override;
+	virtual const FName& GetName() const override;
+	virtual void SetLabel( const FText& InLogLabel ) override;
+	virtual const FText& GetLabel() const override;
+	virtual const TArray< TSharedRef<FMessageFilter> >& GetMessageFilters() const override;
+	virtual void ExecuteToken( const TSharedRef<class IMessageToken>& Token ) const override;
+	virtual void NewPage( const FText& Title ) override;
+	virtual void SetCurrentPage( const FText& Title ) override;
+	virtual void SetCurrentPage( const uint32 InOldPageIndex ) override;
+	virtual void NotifyIfAnyMessages( const FText& Message, EMessageSeverity::Type SeverityFilter = EMessageSeverity::Info, bool bForce = false ) override;
+	virtual void Open() override;
+	virtual int32 NumMessages( EMessageSeverity::Type SeverityFilter ) override;
+
+	DECLARE_DERIVED_EVENT(FMessageLogListingViewModel, IMessageLogListing::IMessageTokenClickedEvent, IMessageTokenClickedEvent)
+	virtual IMessageLogListing::IMessageTokenClickedEvent& OnMessageTokenClicked() override { return TokenClickedEvent; }
+
+	DECLARE_DERIVED_EVENT(FMessageLogListingViewModel, IMessageLogListing::FChangedEvent, FChangedEvent)
+	virtual IMessageLogListing::FChangedEvent& OnDataChanged() override { return ChangedEvent; }
+
+	DECLARE_DERIVED_EVENT(FMessageLogListingViewModel, IMessageLogListing::FOnSelectionChangedEvent, FOnSelectionChangedEvent)
+	virtual IMessageLogListing::FOnSelectionChangedEvent& OnSelectionChanged() override { return SelectionChangedEvent; }
+
+	DECLARE_DERIVED_EVENT(FMessageLogListingViewModel, IMessageLogListing::FOnPageSelectionChangedEvent, FOnPageSelectionChangedEvent )
+	virtual IMessageLogListing::FOnPageSelectionChangedEvent& OnPageSelectionChanged() override { return PageSelectionChangedEvent; }
+	/** End IMessageLogListing interface */
+
+	/** Initializes the FMessageLogListingViewModel for use */
+	virtual void Initialize();
+
+	/** Handles updating the viewmodel when one of its filters changes */
+	void OnFilterChanged();
+
+	/** Called when data is changed changed/updated in the model */
+	virtual void OnChanged();
+
+	/** Returns the filtered message list */
+	const TArray< TSharedRef< FTokenizedMessage > >& GetFilteredMessages() { return FilteredMessages; }
+
+	/** Obtains a const iterator to the filtered messages */
+	MessageContainer::TConstIterator GetFilteredMessageIterator() const;
+
+	/** Obtains a const iterator to the selected filtered messages list */
+	MessageContainer::TConstIterator GetSelectedMessageIterator() const;
+
+	/** Returns the message at the specified index in the filtered list */
+	const TSharedPtr<FTokenizedMessage> GetMessageAtIndex( const int32 MessageIndex ) const;
+
+	/** Set whether we should show filters or not */
+	void SetShowFilters(bool bInShowFilters);
+
+	/** Get whether we should show filters or not */
+	bool GetShowFilters() const;
+
+	/** Set whether we should show pages or not */
+	void SetShowPages(bool bInShowPages);
+
+	/** Get whether we should show pages or not */
+	bool GetShowPages() const;
+
+	/** Set whether we should scroll to the bottom when messages are added */
+	void SetScrollToBottom(bool bInScrollToBottom);
+
+	/** Get whether we should scroll to the bottom when messages are added */
+	bool GetScrollToBottom() const;
+
+	/** Set whether we should show allow the user to clear the log. */
+	void SetAllowClear(bool bInAllowClear);
+
+	/** Get whether we should show allow the user to clear the log. */
+	bool GetAllowClear() const;
+
+	/** Set whether we should discard duplicates or not */
+	void SetDiscardDuplicates(bool bInDiscardDuplicates);
+
+	/** Get whether we should discard duplicates or not */
+	bool GetDiscardDuplicates() const;
+
+	/** Set the maximum page count this log can hold */
+	void SetMaxPageCount(uint32 InMaxPageCount);
+
+	/** Get the maximum page count this log can hold */
+	uint32 GetMaxPageCount() const;
+	
+	/** Get the number of pages we can view */
+	uint32 GetPageCount() const;
+
+	/** Get the current page index we are viewing */
+	uint32 GetCurrentPageIndex() const;
+
+	/** Set the current page index we are viewing */
+	void SetCurrentPageIndex( uint32 InCurrentPageIndex );
+
+	/** Go to the page at the index after the current */
+	void NextPage();
+
+	/** Go to the page at the index before the current */
+	void PrevPage();
+
+	/**
+	 * Get the title of the page at the specified index
+	 * @param	PageIndex	The index of the page
+	 */
+	const FText& GetPageTitle( const uint32 PageIndex ) const;
+
+	/** Gets the number of messages in the current log page */
+	uint32 NumMessages() const;
+
+	/** Get whether to show this log in the main log window */
+	bool ShouldShowInLogWindow() const { return bShowInLogWindow; }
+
+private:
+	FMessageLogListingViewModel( TSharedPtr< FMessageLogListingModel > InMessageLogListingModel, const FText& InLogLabel, const FMessageLogInitializationOptions& InitializationOptions )
+		: bShowFilters( InitializationOptions.bShowFilters )
+		, bShowPages( InitializationOptions.bShowPages )
+		, bAllowClear( InitializationOptions.bAllowClear )
+		, bDiscardDuplicates( InitializationOptions.bDiscardDuplicates )
+		, bScrollToBottom( InitializationOptions.bScrollToBottom )
+		, MaxPageCount( InitializationOptions.MaxPageCount )
+		, bShowInLogWindow( InitializationOptions.bShowInLogWindow )
+		, CurrentPageIndex( 0 )
+		, bIsRefreshing( false )
+		, LogLabel( InLogLabel )
+		, MessageLogListingModel( InMessageLogListingModel )
+	{}
+
+	/** Rebuilds the list of filtered messages */
+	void RefreshFilteredMessages();
+
+	/** Dismisses a notification that was previously shown. */
+	void DismissNotification(int32 NotificationId);
+
+	/** Opens the message log from a notification that was previously shown. */
+	void OpenMessageLogFromNotification(int32 NotificationId);
+
+	/** Helper function for opening this message log from a notification */
+	void OpenMessageLog();
+
+	/** Helper function to check if this log contains messages above a certain severity */
+	int32 NumMessagesPresent( uint32 PageIndex, EMessageSeverity::Type InSeverity ) const;
+
+	/** Helper function to check the worst severity contained in this log page */
+	EMessageSeverity::Type HighestSeverityPresent( uint32 PageIndex ) const;
+
+private:
+
+	/** Whether filters should be shown for this listing */
+	bool bShowFilters;
+
+	/** Whether pages should be used/shown for this listing */
+	bool bShowPages;
+
+	/** Whether we allow the user to clear the log. */
+	bool bAllowClear;
+
+	/** Whether to check for duplicate messages & discard them */
+	bool bDiscardDuplicates;
+
+	/** Whether to scroll to the bottom when messages are added */
+	bool bScrollToBottom;
+
+	/** The limit on the number of displayed pages for this listing */
+	uint32 MaxPageCount;
+
+	/** Whether to show this log in the main log window */
+	bool bShowInLogWindow;
+
+	/** The currently displayed page index */
+	uint32 CurrentPageIndex;
+
+	/** Tracks if the viewmodel is in the middle of refreshing */
+	bool bIsRefreshing;
+
+	/** Label of the listing, displayed to users */
+	FText LogLabel;
+
+	/* The model we are getting display info from */
+	TSharedPtr< FMessageLogListingModel > MessageLogListingModel;
+
+	/** The same list of messages in the model after filtering is applied */
+	MessageContainer FilteredMessages;
+
+	/** The list of selected messages */
+	MessageContainer SelectedFilteredMessages;
+
+	/** The array of message filters used on this listing */
+	TArray< TSharedRef< FMessageFilter > > MessageFilters;
+
+	/** Delegate to call when a token is clicked */
+	IMessageTokenClickedEvent TokenClickedEvent;
+
+	/** Delegate to call when model data is changed */
+	FChangedEvent ChangedEvent;
+
+	/** Delegate to call when selection state is changed */
+	FOnSelectionChangedEvent SelectionChangedEvent;
+
+	/** Delegate to call when page selection state is changed */
+	FOnPageSelectionChangedEvent PageSelectionChangedEvent;
+
+	/** All open notifications */
+	struct FOpenNotification
+	{
+		int32 NotificationId;
+		TWeakPtr<class SNotificationItem> NotificationItem;
+		const FText NotificationMessage;
+
+		FOpenNotification(int32 InNotificationId,
+		                  const TWeakPtr<class SNotificationItem>& InNotificationItem,
+		                  const FText InNotificationMessage)
+			: NotificationId(InNotificationId)
+			, NotificationItem(InNotificationItem)
+			, NotificationMessage( InNotificationMessage)
+		{
+		}
+		FOpenNotification() : NotificationId(INDEX_NONE) {}
+	};
+	TArray<FOpenNotification> OpenNotifications;
+	static int32 NextNotificationId;
+};
diff --git a/Engine/Source/Developer/MessageLog/Public/UserInterface/SMessageLogListing.h b/Engine/Source/Developer/MessageLog/Public/UserInterface/SMessageLogListing.h
new file mode 100644
--- /dev/null
+++ b/Engine/Source/Developer/MessageLog/Public/UserInterface/SMessageLogListing.h
@@ -0,0 +1,167 @@
+// Copyright Epic Games, Inc. All Rights Reserved.
+
+#pragma once
+
+#include "CoreMinimal.h"
+#include "Logging/TokenizedMessage.h"
+#include "Layout/Visibility.h"
+#include "Widgets/DeclarativeSyntaxSupport.h"
+#include "Input/Reply.h"
+#include "Widgets/SWidget.h"
+#include "Widgets/SCompoundWidget.h"
+#include "Widgets/Views/STableViewBase.h"
+#include "Widgets/Views/STableRow.h"
+#include "Presentation/MessageLogListingViewModel.h"
+#include "Framework/Commands/UICommandList.h"
+
+class FMessageFilter;
+
+/**
+ * A message log listing, such as the Compiler Log, or the Map Check Log.
+ * Holds the log lines, and any extra widgets necessary.
+ */
+class MESSAGELOG_API SMessageLogListing : public SCompoundWidget
+{
+public:
+	SLATE_BEGIN_ARGS(SMessageLogListing)
+		{}
+
+	SLATE_END_ARGS()
+
+	SMessageLogListing();
+
+	~SMessageLogListing();
+
+	void Construct( const FArguments& InArgs, const TSharedRef< class IMessageLogListing >& InModelView );
+
+	/** Refreshes the log messages, reapplying any filtering */
+	void RefreshVisibility();
+
+	/** Used to execute the 'on clicked token' delegate */
+	void BroadcastMessageTokenClicked( TSharedPtr< class FTokenizedMessage > Message, const TSharedRef<class IMessageToken>& Token );
+	/** Used to execute a message's featured action (a token associated with the entire message) */
+	void BroadcastMessageDoubleClicked(TSharedPtr< class FTokenizedMessage > Message);
+
+	/** Gets a list of the selected messages */
+	const TArray< TSharedRef< class FTokenizedMessage > > GetSelectedMessages() const;
+
+	/** Set the message selection state */
+	void SelectMessage( const TSharedRef< class FTokenizedMessage >& Message, bool bSelected ) const;
+
+	/** Get the message selection state */
+	bool IsMessageSelected( const TSharedRef< class FTokenizedMessage >& Message ) const;
+
+	/** Scrolls the message into view */
+	void ScrollToMessage( const TSharedRef< class FTokenizedMessage >& Message ) const;
+
+	/** Clears the message selection */
+	void ClearSelectedMessages() const;
+
+	/** Inverts the message selection */
+	void InvertSelectedMessages() const;
+
+	/** Compiles the selected messages into a single string. */
+	FString GetSelectedMessagesAsString() const;
+
+	/** Compiles all the messages into a single string. */
+	FString GetAllMessagesAsString() const;
+
+	/** Gets the message log listing unique name */
+	const FName& GetName() const { return MessageLogListingViewModel->GetName(); }
+
+	/** Gets the message log listing label  */
+	const FText& GetLabel() const { return MessageLogListingViewModel->GetLabel(); }
+
+	/** Gets the set of message filters used when displaying messages */
+	const TArray< TSharedRef< class FMessageFilter > >& GetMessageFilters() const { return MessageLogListingViewModel->GetMessageFilters(); }
+
+	/** Called when data is changed changed/updated in the model */
+	void OnChanged();
+
+	/** Called whenever the viewmodel selection changes */
+	void OnSelectionChanged();
+
+	/**
+	 * Copies the selected messages to the clipboard
+	 */
+	void CopySelectedToClipboard() const;
+
+	/**
+	 * Copies log output (selected/all), optionally to the clipboard
+	 */
+	FString CopyLogOutput( bool bSelected, bool bClipboard ) const;
+
+	/** @return	The UICommandList supported by the MessageLogView */
+	const TSharedRef< const FUICommandList > GetCommandList() const;
+
+	/**
+	 * Called after a key is pressed when this widget has focus
+	 *
+	 * @param  InKeyEvent  Key event
+	 *
+	 * @return  Returns whether the event was handled, along with other possible actions
+	 */
+	FReply OnKeyDown( const FGeometry& MyGeometry, const FKeyEvent& InKeyEvent ) override;
+
+	/** Delegate supplying a label for the page-flipper widget */
+	FString GetPageText() const;
+
+	/** Delegate to display the previous page */
+	FReply OnClickedPrevPage();
+
+	/** Delegate to display the next page */
+	FReply OnClickedNextPage();
+
+	/** Delegate to get the display visibility of the show filter combobox */
+	EVisibility GetFilterMenuVisibility();
+
+	/** Generates the widgets for show filtering */
+	TSharedRef<ITableRow> MakeShowWidget(TSharedRef<FMessageFilter> Selection, const TSharedRef<STableViewBase>& OwnerTable);
+
+	/** Generates the menu content for filtering log content */
+	TSharedRef<SWidget> OnGetFilterMenuContent();
+
+	/** Delegate for enabling & disabling the page widget depending on the number of pages */
+	bool IsPageWidgetEnabled() const;
+
+	/** Delegate for showing & hiding the page widget depending on whether this log uses pages or not */
+	EVisibility GetPageWidgetVisibility() const;
+
+	/** Delegate for enabling & disabling the clear button depending on the number of messages */
+	bool IsClearWidgetEnabled() const;
+
+	/** Delegate for showing & hiding the clear button depending on whether this log uses pages or not */
+	EVisibility GetClearWidgetVisibility() const;
+
+	/** Delegate to generate the label for the page combobox */
+	FText OnGetPageMenuLabel() const;
+
+	/** Delegate to generate the menu content for the page combobox */
+	TSharedRef<SWidget> OnGetPageMenuContent() const;
+
+	/** Delegate called when a page is selected from the page menu */
+	void OnPageSelected(uint32 PageIndex);
+
+	/** Delegate for Clear button */
+	FReply OnClear();
+
+private:
+	/** Makes the message log line widgets for populating the listing */
+	TSharedRef<ITableRow> MakeMessageLogListItemWidget( TSharedRef< class FTokenizedMessage > Message, const TSharedRef< STableViewBase >& OwnerTable );
+	
+	/** Called when map check message line selection has changed */
+	void OnLineSelectionChanged( TSharedPtr< class FTokenizedMessage > Selection, ESelectInfo::Type SelectInfo );
+
+private:
+	/** The list of commands with bound delegates for the message log */
+	const TSharedRef< FUICommandList > UICommandList;
+
+	/** Reference to the ViewModel which holds state info and has access to data */
+	TSharedPtr<class FMessageLogListingViewModel> MessageLogListingViewModel;
+
+	/**	Whether the view is currently updating the viewmodel selection */
+	bool bUpdatingSelection;
+
+	/** The list view for showing all the message log lines */
+	TSharedPtr< SListView< TSharedRef< class FTokenizedMessage > > > MessageListView;
+};
