Allow source link detection in MessageLog text when the source address is located at the right of the message (previously it would only work when the source address was on the left of the message)

diff --git a/Engine/Source/Developer/AutomationController/Private/AutomationControllerManager.cpp b/Engine/Source/Developer/AutomationController/Private/AutomationControllerManager.cpp
--- a/Engine/Source/Developer/AutomationController/Private/AutomationControllerManager.cpp
+++ b/Engine/Source/Developer/AutomationController/Private/AutomationControllerManager.cpp
@@ -1469,24 +1469,27 @@ void FAutomationControllerManager::ReportAutomationResult(const TSharedPtr<IAuto
 
 	for (const FAutomationExecutionEntry& Entry : Results.GetEntries())
 	{
+#if WITH_EDITOR
+		TSharedRef<FTextToken> TextToken = FTextToken::Create(FText::FromString(Entry.ToStringFormattedEditorLog()), false);
+#endif
 		switch (Entry.Event.Type)
 		{
 		case EAutomationEventType::Info:
 			UE_LOG(LogAutomationController, Log, TEXT("%s"), *Entry.ToString());
 #if WITH_EDITOR
-			AutomationEditorLog.Info(FText::FromString(Entry.ToString()));
+			AutomationEditorLog.Info()->AddToken(TextToken);
 #endif
 			break;
 		case EAutomationEventType::Warning:
 			UE_LOG(LogAutomationController, Warning, TEXT("%s"), *Entry.ToString());
 #if WITH_EDITOR
-			AutomationEditorLog.Warning(FText::FromString(Entry.ToString()));
+			AutomationEditorLog.Warning()->AddToken(TextToken);
 #endif
 			break;
 		case EAutomationEventType::Error:
 			UE_LOG(LogAutomationController, Error, TEXT("%s"), *Entry.ToString());
 #if WITH_EDITOR
-			AutomationEditorLog.Error(FText::FromString(Entry.ToString()));
+			AutomationEditorLog.Error()->AddToken(TextToken);
 #endif
 			break;
 		}
diff --git a/Engine/Source/Developer/MessageLog/Private/UserInterface/SMessageLogMessageListRow.cpp b/Engine/Source/Developer/MessageLog/Private/UserInterface/SMessageLogMessageListRow.cpp
--- a/Engine/Source/Developer/MessageLog/Private/UserInterface/SMessageLogMessageListRow.cpp
+++ b/Engine/Source/Developer/MessageLog/Private/UserInterface/SMessageLogMessageListRow.cpp
@@ -256,10 +256,16 @@ void SMessageLogMessageListRow::CreateMessage(const TSharedRef<SHorizontalBox>&
 		else
 		{
 			FString MessageString = InMessageToken->ToText().ToString();
+			const TSharedRef<FTextToken> TextToken = StaticCastSharedRef<FTextToken>(InMessageToken);
 
 			// ^((?:[\w]\:|\\)(?:(?:\\[a-z_\-\s0-9\.]+)+)\.(?:cpp|h))\((\d+)\)
-			// https://regex101.com/r/vV4cV7/1
-			FRegexPattern FileAndLinePattern(TEXT("^((?:[\\w]\\:|\\\\)(?:(?:\\\\[a-z_\\-\\s0-9\\.]+)+)\\.(?:cpp|h))\\((\\d+)\\)"));
+			// https://regex101.com/r/ccAnMS/1
+			FRegexPattern FileAndLinePattern(TEXT("^((?:[\\w]\\:|\\\\)(?:(?:\\\\[A-Za-z_\\-\\s0-9\\.]+)+)\\.(?:cpp|h))\\((\\d+)\\)"));
+			if (!TextToken->IsSourceLinkOnLeft())
+			{
+				FileAndLinePattern = FRegexPattern(TEXT("((?:[\\w]\\:|\\\\)(?:(?:\\\\[A-Za-z_\\-\\s0-9\\.]+)+)\\.(?:cpp|h))\\((\\d+)\\)$"));
+			}
+
 			FRegexMatcher FileAndLineRegexMatcher(FileAndLinePattern, MessageString);
 
 			TSharedRef<SWidget> SourceLink = SNullWidget::NullWidget;
@@ -270,7 +276,14 @@ void SMessageLogMessageListRow::CreateMessage(const TSharedRef<SHorizontalBox>&
 				int32 LineNumber = FCString::Atoi(*FileAndLineRegexMatcher.GetCaptureGroup(2));
 
 				// Remove the hyperlink from the message, since we're splitting it into its own string.
-				MessageString.RightChopInline(FileAndLineRegexMatcher.GetMatchEnding(), false);
+				if (TextToken->IsSourceLinkOnLeft())
+				{
+					MessageString.RightChopInline(FileAndLineRegexMatcher.GetMatchEnding(), false);
+				}
+				else
+				{
+					MessageString.LeftChopInline(FileAndLineRegexMatcher.GetCaptureGroup(0).Len(), false);
+				}
 
 				SourceLink = SNew(SHyperlink)
 					.Style(FAppStyle::Get(), "Common.GotoNativeCodeHyperlink")
@@ -279,22 +292,45 @@ void SMessageLogMessageListRow::CreateMessage(const TSharedRef<SHorizontalBox>&
 					.Text(FText::FromString(FileAndLineRegexMatcher.GetCaptureGroup(0)));
 			}
 
-			RowContent = SNew(SHorizontalBox)
-			+ SHorizontalBox::Slot()
-			.AutoWidth()
-			.Padding(0)
-			[
-				SourceLink
-			]
-			+ SHorizontalBox::Slot()
-			.AutoWidth()
-			.Padding(FMargin(0.f, 4.f))
-			[
-				SNew(STextBlock)
-				.Text(FText::FromString(MessageString))
-				.ColorAndOpacity(FSlateColor::UseForeground())
-				.TextStyle(FAppStyle::Get(), "MessageLog")
-			];
+			if (TextToken->IsSourceLinkOnLeft())
+			{
+				RowContent = SNew(SHorizontalBox)
+					+ SHorizontalBox::Slot()
+					.AutoWidth()
+					.Padding(0)
+					[
+						SourceLink
+					]
+				+ SHorizontalBox::Slot()
+					.AutoWidth()
+					.Padding(FMargin(0.f, 4.f))
+					[
+						SNew(STextBlock)
+						.Text(FText::FromString(MessageString))
+					.ColorAndOpacity(FSlateColor::UseForeground())
+					.TextStyle(FAppStyle::Get(), "MessageLog")
+					];
+			}
+			else
+			{
+				RowContent = SNew(SHorizontalBox)
+				+ SHorizontalBox::Slot()
+					.AutoWidth()
+					.Padding(FMargin(0.f, 4.f))
+					[
+						SNew(STextBlock)
+						.Text(FText::FromString(MessageString))
+					.ColorAndOpacity(FSlateColor::UseForeground())
+					.TextStyle(FAppStyle::Get(), "MessageLog")
+					]
+				+ SHorizontalBox::Slot()
+					.AutoWidth()
+					.Padding(0)
+					[
+						SourceLink
+					];
+			}
+
 		}
 	}
 		break;
diff --git a/Engine/Source/Runtime/Core/Private/Logging/TokenizedMessage.cpp b/Engine/Source/Runtime/Core/Private/Logging/TokenizedMessage.cpp
--- a/Engine/Source/Runtime/Core/Private/Logging/TokenizedMessage.cpp
+++ b/Engine/Source/Runtime/Core/Private/Logging/TokenizedMessage.cpp
@@ -143,9 +143,9 @@ TSharedPtr<IMessageToken> FTokenizedMessage::GetMessageLink() const
 	return MessageLink;
 }
 
-TSharedRef<FTextToken> FTextToken::Create(const FText& InMessage)
+TSharedRef<FTextToken> FTextToken::Create(const FText& InMessage, bool InIsSourceLinkOnLeft)
 {
-	return MakeShared<FTextToken>(FPrivateToken(), InMessage);
+	return MakeShared<FTextToken>(FPrivateToken(), InMessage, InIsSourceLinkOnLeft);
 }
 
 TSharedRef<FDynamicTextToken> FDynamicTextToken::Create(const TAttribute<FText>& InMessage)
diff --git a/Engine/Source/Runtime/Core/Private/Misc/AutomationTest.cpp b/Engine/Source/Runtime/Core/Private/Misc/AutomationTest.cpp
--- a/Engine/Source/Runtime/Core/Private/Misc/AutomationTest.cpp
+++ b/Engine/Source/Runtime/Core/Private/Misc/AutomationTest.cpp
@@ -996,6 +996,30 @@ FString FAutomationExecutionEntry::ToString() const
 	return ComplexString;
 }
 
+FString FAutomationExecutionEntry::ToStringFormattedEditorLog() const
+{
+	FString ComplexString;
+
+	ComplexString = Event.Message;
+
+	if (!Event.Context.IsEmpty())
+	{
+		ComplexString += TEXT(" [");
+		ComplexString += Event.Context;
+		ComplexString += TEXT("] ");
+	}
+
+	if (!Filename.IsEmpty() && LineNumber > 0)
+	{
+		ComplexString += TEXT(" ");
+		ComplexString += Filename;
+		ComplexString += TEXT("(");
+		ComplexString += FString::FromInt(LineNumber);
+		ComplexString += TEXT(")");
+	}
+
+	return ComplexString;
+}
 //------------------------------------------------------------------------------
 
 void FAutomationTestExecutionInfo::Clear()
diff --git a/Engine/Source/Runtime/Core/Public/Logging/TokenizedMessage.h b/Engine/Source/Runtime/Core/Public/Logging/TokenizedMessage.h
--- a/Engine/Source/Runtime/Core/Public/Logging/TokenizedMessage.h
+++ b/Engine/Source/Runtime/Core/Public/Logging/TokenizedMessage.h
@@ -229,7 +229,7 @@ class FTextToken : public IMessageToken
 
 public:
 	/** Factory method, tokens can only be constructed as shared refs */
-	CORE_API static TSharedRef<FTextToken> Create(const FText& InMessage);
+	CORE_API static TSharedRef<FTextToken> Create(const FText& InMessage, bool InIsSourceLinkOnLeft = true);
 
 	/** Begin IMessageToken interface */
 	virtual EMessageToken::Type GetType() const override
@@ -239,10 +239,20 @@ public:
 	/** End IMessageToken interface */
 
 	/** Private constructor */
-	FTextToken(FPrivateToken, const FText& InMessage)
+	FTextToken(FPrivateToken, const FText& InMessage, bool InIsSourceLinkOnLeft)
 	{
 		CachedText = InMessage;
+		bIsSourceLinkOnLeft = InIsSourceLinkOnLeft;
 	}
+
+	bool IsSourceLinkOnLeft() const
+	{
+		return bIsSourceLinkOnLeft;
+	}
+
+private:
+	/** Whether the source address is located on the left or right of the text message. */
+	bool bIsSourceLinkOnLeft;
 };
 
 /** Message token with a localized attribute text payload */
diff --git a/Engine/Source/Runtime/Core/Public/Misc/AutomationEvent.h b/Engine/Source/Runtime/Core/Public/Misc/AutomationEvent.h
--- a/Engine/Source/Runtime/Core/Public/Misc/AutomationEvent.h
+++ b/Engine/Source/Runtime/Core/Public/Misc/AutomationEvent.h
@@ -62,4 +62,5 @@ struct CORE_API FAutomationExecutionEntry
 	}
 
 	FString ToString() const;
+	FString ToStringFormattedEditorLog() const;
 };
