Changed message log action token so that it doesn't disappear when the action has been executed : it becomes disabled instead. Visibility is only tied to the action's Execute callback validity now.
This affects all actions where CanExecuteAction can return false (i.e. single-use actions and actions using the CanExecute delegate)

diff --git a/Engine/Source/Developer/MessageLog/Private/UserInterface/SMessageLogMessageListRow.cpp b/Engine/Source/Developer/MessageLog/Private/UserInterface/SMessageLogMessageListRow.cpp
--- a/Engine/Source/Developer/MessageLog/Private/UserInterface/SMessageLogMessageListRow.cpp
+++ b/Engine/Source/Developer/MessageLog/Private/UserInterface/SMessageLogMessageListRow.cpp
@@ -223,6 +223,7 @@ void SMessageLogMessageListRow::CreateMessage(const TSharedRef<SHorizontalBox>&
 			.Text(InMessageToken->ToText())
 			.ToolTipText(ActionToken->GetActionDescription())
 			.TextStyle(FAppStyle::Get(), "MessageLog")
+			.IsEnabled_Raw(this, &SMessageLogMessageListRow::GetActionLinkEnable, ActionToken)
 			.OnNavigate(this, &SMessageLogMessageListRow::HandleActionHyperlinkNavigate, ActionToken);
 
 		TokenContentVisbility = TAttribute<EVisibility>::Create(TAttribute<EVisibility>::FGetter::CreateRaw(this, &SMessageLogMessageListRow::GetActionLinkVisibility, ActionToken));
diff --git a/Engine/Source/Developer/MessageLog/Private/UserInterface/SMessageLogMessageListRow.h b/Engine/Source/Developer/MessageLog/Private/UserInterface/SMessageLogMessageListRow.h
--- a/Engine/Source/Developer/MessageLog/Private/UserInterface/SMessageLogMessageListRow.h
+++ b/Engine/Source/Developer/MessageLog/Private/UserInterface/SMessageLogMessageListRow.h
@@ -74,7 +74,12 @@ private:
 
 	EVisibility GetActionLinkVisibility(TSharedRef<FActionToken> ActionToken) const
 	{
-		return ActionToken->CanExecuteAction() ? EVisibility::Visible : EVisibility::Collapsed;
+		return ActionToken->IsValidAction() ? EVisibility::Visible : EVisibility::Collapsed;
+	}
+
+	bool GetActionLinkEnable(TSharedRef<FActionToken> ActionToken) const
+	{
+		return ActionToken->CanExecuteAction();
 	}
 
 	void HandleActionHyperlinkNavigate( TSharedRef<FActionToken> ActionToken )
