MessageLog: Remove message logged during MessageLogListingModel::NewPage to remove spammy log messages when cooking projects with many UWorld assets.
#rnx

diff --git a/Engine/Source/Developer/MessageLog/Private/Model/MessageLogListingModel.cpp b/Engine/Source/Developer/MessageLog/Private/Model/MessageLogListingModel.cpp
--- a/Engine/Source/Developer/MessageLog/Private/Model/MessageLogListingModel.cpp
+++ b/Engine/Source/Developer/MessageLog/Private/Model/MessageLogListingModel.cpp
@@ -100,8 +100,6 @@ void FMessageLogListingModel::ClearMessages()
 
 void FMessageLogListingModel::NewPage( const FText& InTitle, uint32 InMaxPages )
 {
-	FMsg::Logf(__FILE__, __LINE__, *LogName.ToString(), ELogVerbosity::Log, TEXT("New page: %s"), *InTitle.ToString());
-
 	// store the name of the page we will add if a new message is pushed
 	PendingPageName = InTitle;
 	MaxPages = InMaxPages;
