SMessageLogListing: Re-enable copying for clang which was disabled in CL #

diff --git a/Engine/Source/Developer/MessageLog/Private/UserInterface/SMessageLogListing.cpp b/Engine/Source/Developer/MessageLog/Private/UserInterface/SMessageLogListing.cpp
--- a/Engine/Source/Developer/MessageLog/Private/UserInterface/SMessageLogListing.cpp
+++ b/Engine/Source/Developer/MessageLog/Private/UserInterface/SMessageLogListing.cpp
@@ -155,14 +155,12 @@ void SMessageLogListing::Construct( const FArguments& InArgs, const TSharedRef<
 	MessageLogListingViewModel->OnDataChanged().AddSP( this, &SMessageLogListing::OnChanged );
 	MessageLogListingViewModel->OnSelectionChanged().AddSP( this, &SMessageLogListing::OnSelectionChanged );
 
-#if !PLATFORM_WINDOWS || !defined(__clang__) // FIXME // @todo clang: this causes an LLVM codegen crash that doesn't go away even with optimizations disabled
 	if (FGenericCommands::IsRegistered())
 	{
 		UICommandList->MapAction(FGenericCommands::Get().Copy,
 			FExecuteAction::CreateSP(this, &SMessageLogListing::CopySelectedToClipboard),
 			FCanExecuteAction());
 	}
-#endif
 }
 
 
