Align the severity and documentation widgets to the top the message log row

This makes it easier to see the start of each row when dealing with multi-line messages, as these widgets used to float to the center of each row which produced an inconsistent appearance that was difficult to visually scan

#rnx

diff --git a/Engine/Source/Developer/MessageLog/Private/UserInterface/SMessageLogMessageListRow.cpp b/Engine/Source/Developer/MessageLog/Private/UserInterface/SMessageLogMessageListRow.cpp
--- a/Engine/Source/Developer/MessageLog/Private/UserInterface/SMessageLogMessageListRow.cpp
+++ b/Engine/Source/Developer/MessageLog/Private/UserInterface/SMessageLogMessageListRow.cpp
@@ -109,7 +109,7 @@ TSharedRef<SWidget> SMessageLogMessageListRow::GenerateWidget()
 		+ SHorizontalBox::Slot()
 		.AutoWidth()
 		.HAlign(HAlign_Center)
-		.VAlign(VAlign_Center)
+		.VAlign(VAlign_Top)
 		[
 			SNew(SBox)
 			.Padding(2.0f)
@@ -131,7 +131,7 @@ TSharedRef<SWidget> SMessageLogMessageListRow::GenerateWidget()
 
 	+ SHorizontalBox::Slot()
 	.FillWidth(1.0f)
-	.VAlign(VAlign_Center)
+	.VAlign(VAlign_Top)
 	[
 		MessageBox
 	]
@@ -139,7 +139,7 @@ TSharedRef<SWidget> SMessageLogMessageListRow::GenerateWidget()
 	+ SHorizontalBox::Slot()
 	.AutoWidth()
 	.HAlign(HAlign_Right)
-	.VAlign(VAlign_Center)
+	.VAlign(VAlign_Top)
 	.Padding(1.0f)
 	[
 		!HasLinks
