Message Log: Fix clearing messages always resulted in the pages being shown

diff --git a/Engine/Source/Developer/MessageLog/Private/Presentation/MessageLogListingViewModel.cpp b/Engine/Source/Developer/MessageLog/Private/Presentation/MessageLogListingViewModel.cpp
--- a/Engine/Source/Developer/MessageLog/Private/Presentation/MessageLogListingViewModel.cpp
+++ b/Engine/Source/Developer/MessageLog/Private/Presentation/MessageLogListingViewModel.cpp
@@ -162,7 +162,11 @@ void FMessageLogListingViewModel::AddMessages( const TArray< TSharedRef< FTokeni
 
 void FMessageLogListingViewModel::ClearMessages()
 {
-	SetCurrentPage(CurrentPageIndex);
+	if (CurrentPageIndex != 0)
+	{
+		SetCurrentPage(CurrentPageIndex);
+	}
+
 	MessageLogListingModel->ClearMessages();
 }
 
