Add ability to disable the MessageLog popup using a CVar
* CVar is "MessageLog.DisablePopup"

diff --git a/Engine/Source/Developer/MessageLog/Private/MessageLogModule.cpp b/Engine/Source/Developer/MessageLog/Private/MessageLogModule.cpp
--- a/Engine/Source/Developer/MessageLog/Private/MessageLogModule.cpp
+++ b/Engine/Source/Developer/MessageLog/Private/MessageLogModule.cpp
@@ -25,6 +25,14 @@
 
 IMPLEMENT_MODULE( FMessageLogModule, MessageLog );
 
+static bool GDisableMessageLogPopup = false;
+FAutoConsoleVariableRef CVarDisableMessageLogPopup(
+	TEXT("MessageLog.DisablePopup"),
+	GDisableMessageLogPopup,
+	TEXT("Disables the MessageLog popup"),
+	ECVF_Default
+);
+
 TSharedRef<SDockTab> SpawnMessageLog( const FSpawnTabArgs& Args, TSharedRef<FMessageLogViewModel> MessageLogViewModel )
 {
 	LLM_SCOPE(ELLMTag::UI);
@@ -117,6 +125,10 @@ TSharedRef<SWidget> FMessageLogModule::CreateLogListingWidget(const TSharedRef<I
 
 void FMessageLogModule::OpenMessageLog(const FName& LogName)
 {
+	if (GDisableMessageLogPopup)
+	{
+		return;
+	}
 	// only open the message log if we have a window created for the tab manager & our delegate allows it
 	if(bCanDisplayMessageLog
 #if !PLATFORM_MAC
