Message Log: Improve UX.

Draw a shadow if there are more elements that are currently visible.
Make the rows hoverable and clickable through the whole width of the panel.
Decrease padding to waste less space when the log listing is used as a docked tab.

diff --git a/Engine/Source/Developer/MessageLog/Private/UserInterface/SMessageLogListing.cpp b/Engine/Source/Developer/MessageLog/Private/UserInterface/SMessageLogListing.cpp
--- a/Engine/Source/Developer/MessageLog/Private/UserInterface/SMessageLogListing.cpp
+++ b/Engine/Source/Developer/MessageLog/Private/UserInterface/SMessageLogListing.cpp
@@ -9,6 +9,7 @@
 #include "Widgets/Input/SButton.h"
 #include "Widgets/Input/SComboButton.h"
 #include "Widgets/Views/SListView.h"
+#include "Widgets/Layout/SScrollBorder.h"
 #include "Widgets/Layout/SScrollBox.h"
 #include "Widgets/Input/SCheckBox.h"
 #include "UserInterface/SMessageLogMessageListRow.h"
@@ -42,11 +43,18 @@ void SMessageLogListing::Construct( const FArguments& InArgs, const TSharedRef<
 
 	TSharedRef<SScrollBar> ScrollBar = SNew(SScrollBar);
 
+	MessageListView = SNew(SListView<TSharedRef<FTokenizedMessage>>)
+		.ListItemsSource(&MessageLogListingViewModel->GetFilteredMessages())
+		.OnGenerateRow(this, &SMessageLogListing::MakeMessageLogListItemWidget)
+		.OnSelectionChanged(this, &SMessageLogListing::OnLineSelectionChanged)
+		.ExternalScrollbar(ScrollBar)
+		.ConsumeMouseWheel(EConsumeMouseWheel::Always);
+
 	ChildSlot
 	[
 		SNew(SBorder)
 		.BorderImage(FAppStyle::Get().GetBrush("Brushes.Recessed"))
-		.Padding(15.0f)
+		.Padding(6.0f)
 		[
 			SAssignNew(VerticalBox, SVerticalBox)
 			+ SVerticalBox::Slot()
@@ -60,13 +68,12 @@ void SMessageLogListing::Construct( const FArguments& InArgs, const TSharedRef<
 					SNew(SScrollBox)
 					.Orientation(EOrientation::Orient_Horizontal)
 					+ SScrollBox::Slot()
+					.FillSize(1.0f)
 					[
-						SAssignNew(MessageListView, SListView<TSharedRef<FTokenizedMessage>>)
-						.ListItemsSource(&MessageLogListingViewModel->GetFilteredMessages())
-						.OnGenerateRow(this, &SMessageLogListing::MakeMessageLogListItemWidget)
-						.OnSelectionChanged(this, &SMessageLogListing::OnLineSelectionChanged)
-						.ExternalScrollbar(ScrollBar)
-						.ConsumeMouseWheel(EConsumeMouseWheel::Always)
+						SNew(SScrollBorder, MessageListView.ToSharedRef())
+						[
+							MessageListView.ToSharedRef()
+						]
 					]
 				]
 				+SHorizontalBox::Slot()
