Allow filtering the message log by default

diff --git a/Engine/Source/Developer/MessageLog/Public/MessageLogInitializationOptions.h b/Engine/Source/Developer/MessageLog/Public/MessageLogInitializationOptions.h
--- a/Engine/Source/Developer/MessageLog/Public/MessageLogInitializationOptions.h
+++ b/Engine/Source/Developer/MessageLog/Public/MessageLogInitializationOptions.h
@@ -10,7 +10,7 @@
 struct FMessageLogInitializationOptions
 {
 	FMessageLogInitializationOptions()
-		: bShowFilters(false)
+		: bShowFilters(true)
 		, bShowPages(false)
 		, bAllowClear(true)
 		, bDiscardDuplicates(false)
