Map Check log window support for multi-select of objects affected by errors, via right click context menu.  Includes option to select actors from multiple selected messages, or to select all objects with a specific error, or specific error and mesh.  Intermediate step towards making it easier to solve StaticMesh Map Check errors in bulk (UE - 243386).

diff --git a/Engine/Source/Developer/MessageLog/Private/UserInterface/SMessageLogListing.cpp b/Engine/Source/Developer/MessageLog/Private/UserInterface/SMessageLogListing.cpp
--- a/Engine/Source/Developer/MessageLog/Private/UserInterface/SMessageLogListing.cpp
+++ b/Engine/Source/Developer/MessageLog/Private/UserInterface/SMessageLogListing.cpp
@@ -16,6 +16,7 @@
 #include "Framework/Commands/GenericCommands.h"
 #include "HAL/PlatformApplicationMisc.h"
 #include "Styling/StyleColors.h"
+#include "Logging/MessageLog.h"
 
 
 #define LOCTEXT_NAMESPACE "Developer.MessageLog"
@@ -48,7 +49,8 @@ void SMessageLogListing::Construct( const FArguments& InArgs, const TSharedRef<
 		.OnGenerateRow(this, &SMessageLogListing::MakeMessageLogListItemWidget)
 		.OnSelectionChanged(this, &SMessageLogListing::OnLineSelectionChanged)
 		.ExternalScrollbar(ScrollBar)
-		.ConsumeMouseWheel(EConsumeMouseWheel::Always);
+		.ConsumeMouseWheel(EConsumeMouseWheel::Always)
+		.OnContextMenuOpening(this, &SMessageLogListing::ContextMenuOpening);
 
 	ChildSlot
 	[
@@ -298,6 +300,155 @@ TSharedRef<ITableRow> SMessageLogListing::MakeMessageLogListItemWidget( TSharedR
 		.OnMessageDoubleClicked( this, &SMessageLogListing::BroadcastMessageDoubleClicked );
 }
 
+static bool CanMessageTokenBeActor(EMessageToken::Type TokenType)
+{
+	return (TokenType == EMessageToken::Actor || TokenType == EMessageToken::Object || TokenType == EMessageToken::AssetData);
+}
+
+static void ActivateMessageToken(const TSharedRef<FTokenizedMessage>& Message)
+{
+	const TArray<TSharedRef<IMessageToken>>& MessageTokens = Message->GetMessageTokens();
+	for (const TSharedRef<IMessageToken>& Token : MessageTokens)
+	{
+		if (CanMessageTokenBeActor(Token->GetType()))
+		{
+			const FOnMessageTokenActivated& ActivateDelegate = Token->GetOnMessageTokenActivated();
+			if (ActivateDelegate.IsBound())
+			{
+				ActivateDelegate.Execute(Token);
+				return;
+			}
+		}
+	}
+}
+
+// The source string for a formatted message is the original format specifier string, without any substitutions.  Optional FormatKey
+// can be used to also return the value of a specific field in the formatted string.
+static FString GetSourceStringFromMessage(const TArray<TSharedRef<IMessageToken>>& MessageTokens, const TCHAR* FormatKey, FString& OutFormatValue)
+{
+	OutFormatValue.Empty();
+
+	const TSharedRef<IMessageToken>* TextToken = MessageTokens.FindByPredicate([](const TSharedRef<IMessageToken>& Token) { return Token->GetType() == EMessageToken::Text; });
+	if (TextToken)
+	{
+		const FText& Text = (*TextToken)->ToText();
+		TArray<FHistoricTextFormatData> TextHistory;
+		FTextInspector::GetHistoricFormatData(Text, TextHistory);
+
+		if (TextHistory.Num() > 0)
+		{
+			if (FormatKey)
+			{
+				const FFormatArgumentValue* FoundValue = TextHistory[0].Arguments.Find(FormatKey);
+				if (FoundValue)
+				{
+					OutFormatValue = FoundValue->GetTextValue().ToString();
+				}
+			}
+
+			return TextHistory[0].SourceFmt.GetSourceString();
+		}
+	}
+
+	return FString();
+}
+
+void SMessageLogListing::ContextMenuSelect() const
+{
+	FMessageLog::OnMultiSelectActorBegin().ExecuteIfBound();
+	TArray<TSharedRef<FTokenizedMessage>> SelectedMessages = GetSelectedMessages();
+	for (TSharedRef<FTokenizedMessage>& Message : SelectedMessages)
+	{
+		ActivateMessageToken(Message);
+	}
+	FMessageLog::OnMultiSelectActorEnd().ExecuteIfBound();
+}
+
+// Optional FormatKey requires the value of the given field in the format specifier to match across all messages.
+void SMessageLogListing::ContextMenuSelectByMessage(const TCHAR* FormatKey) const
+{
+	TArray<TSharedRef<FTokenizedMessage>> SelectedMessages = GetSelectedMessages();
+	if (SelectedMessages.Num() == 1)
+	{
+		TSharedRef<FTokenizedMessage>& SelectedMessage = SelectedMessages[0];
+		const TArray<TSharedRef<IMessageToken>>& SelectedMessageTokens = SelectedMessage->GetMessageTokens();
+		FString SelectedFormatValue;
+		FString SelectedSourceString = GetSourceStringFromMessage(SelectedMessageTokens, FormatKey, SelectedFormatValue);
+
+		if (!SelectedSourceString.IsEmpty() && (FormatKey == nullptr || !SelectedFormatValue.IsEmpty()) &&
+			SelectedMessageTokens.ContainsByPredicate([](const TSharedRef<IMessageToken>& Token) { return CanMessageTokenBeActor(Token->GetType()); }))
+		{
+			FMessageLog::OnMultiSelectActorBegin().ExecuteIfBound();
+			for (const TSharedRef<FTokenizedMessage>& MessageItem : MessageListView->GetItems())
+			{
+				const TArray<TSharedRef<IMessageToken>>& MessageItemTokens = MessageItem->GetMessageTokens();
+				FString MessageItemFormatValue;
+				FString MessageItemSourceString = GetSourceStringFromMessage(MessageItemTokens, FormatKey, MessageItemFormatValue);
+
+				if (MessageItemSourceString == SelectedSourceString && MessageItemFormatValue == SelectedFormatValue)
+				{
+					ActivateMessageToken(MessageItem);
+				}
+			}
+			FMessageLog::OnMultiSelectActorEnd().ExecuteIfBound();
+		}
+	}
+}
+
+TSharedPtr<SWidget> SMessageLogListing::ContextMenuOpening()
+{
+	// Only show actor selection context menu for map check message log
+	if (MessageLogListingViewModel->GetLabel().BuildSourceString() == TEXT("Map Check"))
+	{
+		// Search for a message with an actor, object, or asset token
+		TArray<TSharedRef<FTokenizedMessage>> SelectedMessages = GetSelectedMessages();
+		for (TSharedRef<FTokenizedMessage>& Message : SelectedMessages)
+		{
+			const TArray<TSharedRef<IMessageToken>>& MessageTokens = Message->GetMessageTokens();
+			if (MessageTokens.ContainsByPredicate([](const TSharedRef<IMessageToken>& Token) { return CanMessageTokenBeActor(Token->GetType()); }))
+			{
+				// Found a relevant message, start building the context menu
+				FMenuBuilder MenuBuilder(true, nullptr);
+
+				// Batch selection by message type or mesh name is only supported for single selection, for simplicity and clarity
+				if (SelectedMessages.Num() == 1)
+				{
+					MenuBuilder.AddMenuEntry(LOCTEXT("SelectActor", "Select actor"), LOCTEXT("SelectActorTooltip", "Select actor referenced by message"), FSlateIcon(),
+						FUIAction(FExecuteAction::CreateRaw(this, &SMessageLogListing::ContextMenuSelect)));
+
+					const TSharedRef<IMessageToken>* TextToken = MessageTokens.FindByPredicate([](const TSharedRef<IMessageToken>& Token) { return Token->GetType() == EMessageToken::Text; });
+					if (TextToken)
+					{
+						const FText& Text = (*TextToken)->ToText();
+						TArray<FHistoricTextFormatData> TextHistory;
+						FTextInspector::GetHistoricFormatData(Text, TextHistory);
+
+						if (TextHistory.Num() > 0)
+						{
+							MenuBuilder.AddMenuEntry(LOCTEXT("SelectActorsByMessage", "Select actors by message"), LOCTEXT("SelectActorsByMessageTooltip", "Select actors with the same map check message"), FSlateIcon(),
+								FUIAction(FExecuteAction::CreateRaw(this, &SMessageLogListing::ContextMenuSelectByMessage, (const TCHAR*)nullptr)));
+
+							if (TextHistory[0].Arguments.Contains(TEXT("MeshName")))
+							{
+								MenuBuilder.AddMenuEntry(LOCTEXT("SelectActorsByMessageAndMesh", "Select actors by message / mesh"), LOCTEXT("SelectActorsByMessageAndMeshTooltip", "Select actors with the same map check message and mesh"), FSlateIcon(),
+									FUIAction(FExecuteAction::CreateRaw(this, &SMessageLogListing::ContextMenuSelectByMessage, (const TCHAR*)TEXT("MeshName"))));
+							}
+						}
+					}
+				}
+				else
+				{
+					MenuBuilder.AddMenuEntry(LOCTEXT("SelectActors", "Select actors"), LOCTEXT("SelectActorsTooltip", "Select actors referenced by all selected messages"), FSlateIcon(),
+						FUIAction(FExecuteAction::CreateRaw(this, &SMessageLogListing::ContextMenuSelect)));
+				}
+
+				return MenuBuilder.MakeWidget();
+			}
+		}
+	}
+
+	return nullptr;
+}
 
 void SMessageLogListing::OnLineSelectionChanged( TSharedPtr< FTokenizedMessage > Selection, ESelectInfo::Type /*SelectInfo*/ )
 {
diff --git a/Engine/Source/Developer/MessageLog/Public/UserInterface/SMessageLogListing.h b/Engine/Source/Developer/MessageLog/Public/UserInterface/SMessageLogListing.h
--- a/Engine/Source/Developer/MessageLog/Public/UserInterface/SMessageLogListing.h
+++ b/Engine/Source/Developer/MessageLog/Public/UserInterface/SMessageLogListing.h
@@ -154,6 +154,10 @@ private:
 	/** Called when map check message line selection has changed */
 	UE_API void OnLineSelectionChanged( TSharedPtr< class FTokenizedMessage > Selection, ESelectInfo::Type SelectInfo );
 
+	void ContextMenuSelect() const;
+	void ContextMenuSelectByMessage(const TCHAR* FormatKey) const;
+	TSharedPtr<SWidget> ContextMenuOpening();
+
 private:
 	/** The list of commands with bound delegates for the message log */
 	const TSharedRef< FUICommandList > UICommandList;
diff --git a/Engine/Source/Editor/UnrealEd/Private/UnrealEdMisc.cpp b/Engine/Source/Editor/UnrealEd/Private/UnrealEdMisc.cpp
--- a/Engine/Source/Editor/UnrealEd/Private/UnrealEdMisc.cpp
+++ b/Engine/Source/Editor/UnrealEd/Private/UnrealEdMisc.cpp
@@ -558,6 +558,8 @@ void FUnrealEdMisc::OnInit()
 
 	// install message log delegates
 	FMessageLog::OnMessageSelectionChanged().BindRaw(this, &FUnrealEdMisc::OnMessageSelectionChanged);
+	FMessageLog::OnMultiSelectActorBegin().BindRaw(this, &FUnrealEdMisc::OnMultiSelectActorFromMessageTokenBegin);
+	FMessageLog::OnMultiSelectActorEnd().BindRaw(this, &FUnrealEdMisc::OnMultiSelectActorFromMessageTokenEnd);
 	FUObjectToken::DefaultOnMessageTokenActivated().BindRaw(this, &FUnrealEdMisc::OnMessageTokenActivated);
 	FUObjectToken::DefaultOnGetObjectDisplayName().BindRaw(this, &FUnrealEdMisc::OnGetDisplayName);
 	FAssetNameToken::OnGotoAsset().BindRaw(this, &FUnrealEdMisc::OnGotoAsset);
@@ -1080,6 +1082,8 @@ void FUnrealEdMisc::OnExit()
 
 	FInputBindingManager::Get().UnregisterUserDefinedChordChanged(OnUserDefinedChordChangedDelegateHandle);
 	FMessageLog::OnMessageSelectionChanged().Unbind();
+	FMessageLog::OnMultiSelectActorBegin().Unbind();
+	FMessageLog::OnMultiSelectActorEnd().Unbind();
 	FUObjectToken::DefaultOnMessageTokenActivated().Unbind();
 	FUObjectToken::DefaultOnGetObjectDisplayName().Unbind();
 	FAssetNameToken::OnGotoAsset().Unbind();
@@ -1417,6 +1421,12 @@ void FUnrealEdMisc::OnMessageTokenActivated(const TSharedRef<IMessageToken>& Tok
 		ULightmappedSurfaceCollection* SurfaceCollection = Cast<ULightmappedSurfaceCollection>(Object);
 		if (SurfaceCollection)
 		{
+			// Ignore surface collection for actor multi-select
+			if (bActorMultiSelect)
+			{
+				return;
+			}
+
 			// Deselect all selected object...
 			GEditor->SelectNone( true, true );
 
@@ -1456,6 +1466,12 @@ void FUnrealEdMisc::OnMessageTokenActivated(const TSharedRef<IMessageToken>& Tok
 			}
 			else
 			{
+				// Ignore blueprint objects for actor multi-select
+				if (bActorMultiSelect)
+				{
+					return;
+				}
+
 				TArray<UObject*> ObjectArray;
 
 				if (Object->IsInBlueprint())
@@ -1529,6 +1545,12 @@ void FUnrealEdMisc::OnAssetDataTokenActivated(const TSharedRef<class IMessageTok
 	// If this is a standalone asset, jump to it in the content browser
 	if (AssetData.GetOptionalOuterPathName().IsNone())
 	{
+		// Actor multi-select only cares about actors, not assets
+		if (bActorMultiSelect)
+		{
+			return;
+		}
+
 		GEditor->SyncBrowserToObjects({ AssetData });
 	}
 	// If this is a loaded actor, resolve the pointer and select it
@@ -1546,12 +1568,26 @@ void FUnrealEdMisc::OnAssetDataTokenActivated(const TSharedRef<class IMessageTok
 			TUniquePtr<FWorldPartitionActorDesc> Desc = FWorldPartitionActorDescUtils::GetActorDescriptorFromAssetData(AssetData);
 			if (Desc.IsValid())
 			{
-				GEditor->BroadcastSelectUnloadedActors({ Desc->GetGuid() });
+				if (bActorMultiSelect)
+				{
+					// Selection is deferred to OnMultiSelectActorEnd when multi-select is active
+					UnloadedActorsSelected.AddUnique(Desc->GetGuid());
+				}
+				else
+				{
+					GEditor->BroadcastSelectUnloadedActors({ Desc->GetGuid() });
+				}
 			}
 		}
 		// If this is an actor in an unloaded level, jump to the level in the content browser
 		else 
 		{
+			// Actor multi-select doesn't care about browsing to unloaded levels
+			if (bActorMultiSelect)
+			{
+				return;
+			}
+
 			FAssetRegistryModule& AssetRegistryModule = FModuleManager::GetModuleChecked<FAssetRegistryModule>("AssetRegistry");
 			IAssetRegistry& AssetRegistry = AssetRegistryModule.Get();
 			TArray<FAssetData> OutAssets;
@@ -1572,13 +1608,28 @@ void FUnrealEdMisc::OnActorTokenActivated(const TSharedRef<class IMessageToken>&
 		}
 		else
 		{
-			GEditor->BroadcastSelectUnloadedActors({ ActorToken->GetActorGuid() });
+			if (bActorMultiSelect)
+			{
+				// Selection is deferred to OnMultiSelectActorEnd when multi-select is active
+				UnloadedActorsSelected.AddUnique(ActorToken->GetActorGuid());
+			}
+			else
+			{
+				GEditor->BroadcastSelectUnloadedActors({ ActorToken->GetActorGuid() });
+			}
 		}
 	}
 }
 
 void FUnrealEdMisc::SelectActorFromMessageToken(AActor* InActor)
 {
+	if (bActorMultiSelect)
+	{
+		// Selection is deferred to OnMultiSelectActorEnd when multi-select is active
+		ActorsSelected.AddUnique(InActor);
+		return;
+	}
+
 	// Select the actor
 	GEditor->SelectNone(false, true);
 	GEditor->SelectActor(InActor, /*InSelected=*/true, /*bNotify=*/false, /*bSelectEvenIfHidden=*/true);
@@ -1590,6 +1641,44 @@ void FUnrealEdMisc::SelectActorFromMessageToken(AActor* InActor)
 	GUnrealEd->UpdateFloatingPropertyWindows();
 }
 
+UNREALED_API void FUnrealEdMisc::OnMultiSelectActorFromMessageTokenBegin()
+{
+	check(bActorMultiSelect == false);
+	bActorMultiSelect = true;
+}
+
+UNREALED_API void FUnrealEdMisc::OnMultiSelectActorFromMessageTokenEnd()
+{
+	check(bActorMultiSelect == true);
+	bActorMultiSelect = false;
+
+	// Clear selection
+	GEditor->SelectNone(false, true);
+
+	// Notify about selection of unloaded actors
+	if (UnloadedActorsSelected.Num())
+	{
+		GEditor->BroadcastSelectUnloadedActors(UnloadedActorsSelected);
+		UnloadedActorsSelected.Empty();
+	}
+
+	// Select the actors
+	if (ActorsSelected.Num())
+	{
+		for (AActor* Actor : ActorsSelected)
+		{
+			GEditor->SelectActor(Actor, /*InSelected=*/true, /*bNotify=*/false, /*bSelectEvenIfHidden=*/true);
+		}
+		GEditor->NoteSelectionChange();
+		GEditor->MoveViewportCamerasToActor(ActorsSelected, false);
+		ActorsSelected.Empty();
+	}
+
+	// Update the property windows and create one if necessary
+	GUnrealEd->ShowActorProperties();
+	GUnrealEd->UpdateFloatingPropertyWindows();
+}
+
 FText FUnrealEdMisc::OnGetDisplayName(const UObject* InObject, const bool bFullPath)
 {
 	FText Name = LOCTEXT("DisplayNone", "<None>");
diff --git a/Engine/Source/Editor/UnrealEd/Public/UnrealEdMisc.h b/Engine/Source/Editor/UnrealEd/Public/UnrealEdMisc.h
--- a/Engine/Source/Editor/UnrealEd/Public/UnrealEdMisc.h
+++ b/Engine/Source/Editor/UnrealEd/Public/UnrealEdMisc.h
@@ -286,6 +286,10 @@ public:
 	UNREALED_API bool SpawnEditorInstance(const FString& ProjectName);
 
 private:
+	/** Delegates to support multi-select of message tokens.  Selections from calls to OnMessageTokenActivated or OnActorTokenActivated are batched and selected as a group. */
+	UNREALED_API void OnMultiSelectActorFromMessageTokenBegin();
+	UNREALED_API void OnMultiSelectActorFromMessageTokenEnd();
+
 	UNREALED_API void SelectActorFromMessageToken(AActor* InActor);
 
 	UNREALED_API void PreSaveWorld(class UWorld* World, FObjectPreSaveContext ObjectSaveContext);
@@ -348,4 +352,8 @@ private:
 	FDelegateHandle OnTabForegroundedDelegateHandle;
 
 	FTimerHandle EditorAnalyticsHeartbeatTimerHandle;	
+
+	bool bActorMultiSelect = false;
+	TArray<AActor*> ActorsSelected;
+	TArray<FGuid> UnloadedActorsSelected;
 };
diff --git a/Engine/Source/Runtime/Core/Private/Logging/MessageLog.cpp b/Engine/Source/Runtime/Core/Private/Logging/MessageLog.cpp
--- a/Engine/Source/Runtime/Core/Private/Logging/MessageLog.cpp
+++ b/Engine/Source/Runtime/Core/Private/Logging/MessageLog.cpp
@@ -9,6 +9,8 @@
 FMessageLog::FGetLog FMessageLog::GetLog;
 
 FMessageLog::FMessageSelectionChanged FMessageLog::MessageSelectionChanged;
+FMessageLog::FMultiSelectActorBegin FMessageLog::MultiSelectActorBegin;
+FMessageLog::FMultiSelectActorEnd FMessageLog::MultiSelectActorEnd;
 
 #define LOCTEXT_NAMESPACE "MessageLog"
 
diff --git a/Engine/Source/Runtime/Core/Public/Logging/MessageLog.h b/Engine/Source/Runtime/Core/Public/Logging/MessageLog.h
--- a/Engine/Source/Runtime/Core/Public/Logging/MessageLog.h
+++ b/Engine/Source/Runtime/Core/Public/Logging/MessageLog.h
@@ -127,6 +127,13 @@ public:
 	/** Helper function to retrieve a log color for a specified message severity */
 	CORE_API static const TCHAR* const GetLogColor( EMessageSeverity::Type InSeverity );
 
+	/** Delegates for multi-select of message token actors */
+	DECLARE_DELEGATE(FMultiSelectActorBegin);
+	static FMultiSelectActorBegin& OnMultiSelectActorBegin() { return MultiSelectActorBegin; }
+
+	DECLARE_DELEGATE(FMultiSelectActorEnd);
+	static FMultiSelectActorEnd& OnMultiSelectActorEnd() { return MultiSelectActorEnd; }
+
 private:
 	/** Buffer for messages */
 	TArray< TSharedRef<FTokenizedMessage> > Messages;
@@ -145,6 +152,10 @@ private:
 
 	/** Delegate for message selection */
 	CORE_API static FMessageSelectionChanged MessageSelectionChanged;
+
+	/** Delegates for multi-select of message token actors */
+	CORE_API static FMultiSelectActorBegin MultiSelectActorBegin;
+	CORE_API static FMultiSelectActorEnd MultiSelectActorEnd;
 };
 
 /**
