Don't bind the visual MessageLog implementation for commandlets (as they have no UI)

#rnx
#Synced-CL 51332317

diff --git a/Engine/Source/Developer/MessageLog/Private/MessageLogModule.cpp b/Engine/Source/Developer/MessageLog/Private/MessageLogModule.cpp
--- a/Engine/Source/Developer/MessageLog/Private/MessageLogModule.cpp
+++ b/Engine/Source/Developer/MessageLog/Private/MessageLogModule.cpp
@@ -70,8 +70,11 @@ void FMessageLogModule::StartupModule()
 	});
 #endif
 
-	// Bind us so message log output is routed via this module
-	FMessageLog::OnGetLog().BindStatic(&GetLog);
+	if (!IsRunningCommandlet()) // Commandlets can just use FBasicMessageLog from Core
+	{
+		// Bind us so message log output is routed via this module
+		FMessageLog::OnGetLog().BindStatic(&GetLog);
+	}
 }
 
 void FMessageLogModule::ShutdownModule()
@@ -84,7 +87,10 @@ void FMessageLogModule::ShutdownModule()
 	FModuleManager::Get().OnModulesChanged().Remove(ModulesChangedHandle);
 #endif
 
-	FMessageLog::OnGetLog().Unbind();
+	if (!IsRunningCommandlet()) // Commandlets can just use FBasicMessageLog from Core
+	{
+		FMessageLog::OnGetLog().Unbind();
+	}
 }
 
 TSharedRef<IMessageLogListing> FMessageLogModule::GetLogListing(const FName& LogName)
